Renamed UE4ML plugin to MLAdapter and updated all its contents accordingly

diff --git a/Engine/Plugins/AI/MLAdapter/MLAdapter.uplugin b/Engine/Plugins/AI/MLAdapter/MLAdapter.uplugin
new file mode 100644
--- /dev/null
+++ b/Engine/Plugins/AI/MLAdapter/MLAdapter.uplugin
@@ -0,0 +1,35 @@
+{
+	"FileVersion": 3,
+	"Version": 1,
+	"VersionName": "0.0.1",
+	"FriendlyName": "Unreal Engine support for Machine Learning",
+	"Description": "[EXPERIMENTAL] Exposes means of having external agents query world and game state and control in-game agents",
+	"Category": "AI",
+	"CreatedBy": "Epic Games, Inc.",
+	"CreatedByURL": "http://epicgames.com",
+	"DocsURL": "",
+	"MarketplaceURL": "",
+	"SupportURL": "",
+	"EnabledByDefault": false,
+	"CanContainContent": false,
+	"IsBetaVersion": true,
+	"Installed": false,
+	"Modules": [
+		{
+			"Name": "MLAdapter",
+			"Type": "Runtime",
+			"LoadingPhase": "PreDefault"
+		},
+		{
+			"Name": "MLAdapterTestSuite",
+			"Type": "DeveloperTool",
+			"LoadingPhase": "PreDefault"
+		}
+	],
+	"Plugins": [
+		{
+			"Name": "GameplayAbilities",
+			"Enabled": true
+		}
+	]
+}
\ No newline at end of file
diff --git a/Engine/Plugins/AI/MLAdapter/Source/MLAdapter/MLAdapter.Build.cs b/Engine/Plugins/AI/MLAdapter/Source/MLAdapter/MLAdapter.Build.cs
new file mode 100644
--- /dev/null
+++ b/Engine/Plugins/AI/MLAdapter/Source/MLAdapter/MLAdapter.Build.cs
@@ -0,0 +1,75 @@
+// Copyright Epic Games, Inc. All Rights Reserved.
+
+using System.IO;
+
+namespace UnrealBuildTool.Rules
+{
+    public class MLAdapter : ModuleRules
+    {
+        public MLAdapter(ReadOnlyTargetRules Target) : base(Target)
+        {
+            // rcplib is using exceptions so we have to enable that
+            bEnableExceptions = true;
+            PCHUsage = ModuleRules.PCHUsageMode.UseExplicitOrSharedPCHs;
+
+            PublicIncludePaths.AddRange(
+                new string[] {
+                "Runtime/AIModule/Public",
+                ModuleDirectory + "/Public",
+                }
+            );
+
+            PrivateIncludePaths.Add("MLAdapter/Private");
+
+            PublicDependencyModuleNames.AddRange(
+                new string[] {
+                    "Core",
+                    "CoreUObject",
+                    "Engine",
+                    "GameplayTags",
+                    "AIModule",
+                    "InputCore",
+                    "Json",
+                    "JsonUtilities",
+                    "GameplayAbilities",
+                }
+            );
+
+            PrivateDependencyModuleNames.AddRange(
+                new string[] {
+                    "DeveloperSettings"
+                }
+            );
+
+            // RPCLib disabled on other platforms at the moment
+            if (Target.Platform == UnrealTargetPlatform.Win64)
+            {
+                PublicDefinitions.Add("WITH_RPCLIB=1");
+                AddEngineThirdPartyPrivateStaticDependencies(Target, "RPCLib");
+            
+				string RPClibDir = Path.Combine(Target.UEThirdPartySourceDirectory, "rpclib");
+				PublicIncludePaths.Add(Path.Combine(RPClibDir, "Source", "include"));
+            }
+            else
+            {
+                PublicDefinitions.Add("WITH_RPCLIB=0");
+            }
+
+            if (Target.bBuildEditor == true)
+            {
+                PrivateDependencyModuleNames.Add("EditorFramework");
+                PrivateDependencyModuleNames.Add("UnrealEd");
+            }
+
+            if (Target.bBuildDeveloperTools || (Target.Configuration != UnrealTargetConfiguration.Shipping && Target.Configuration != UnrealTargetConfiguration.Test))
+            {
+                PrivateDependencyModuleNames.Add("GameplayDebugger");
+                PublicDefinitions.Add("WITH_GAMEPLAY_DEBUGGER=1");
+            }
+            else
+            {
+                PublicDefinitions.Add("WITH_GAMEPLAY_DEBUGGER=0");
+            }
+        }
+    }
+}
diff --git a/Engine/Plugins/AI/MLAdapter/Source/MLAdapter/Private/Agents/MLAdapterAgentElements.cpp b/Engine/Plugins/AI/MLAdapter/Source/MLAdapter/Private/Agents/MLAdapterAgentElements.cpp
new file mode 100644
--- /dev/null
+++ b/Engine/Plugins/AI/MLAdapter/Source/MLAdapter/Private/Agents/MLAdapterAgentElements.cpp
@@ -0,0 +1,106 @@
+// Copyright Epic Games, Inc. All Rights Reserved.
+
+#include "Agents/MLAdapterAgentElement.h"
+#include "GameFramework/Controller.h"
+#include "GameFramework/Pawn.h"
+#include "Agents/MLAdapterAgent.h"
+
+
+UMLAdapterAgentElement::UMLAdapterAgentElement(const FObjectInitializer& ObjectInitializer)
+	: Super(ObjectInitializer)
+	, SpaceDef(MakeShareable(new FMLAdapter::FSpace_Dummy()))
+{
+	if (HasAnyFlags(RF_ClassDefaultObject))
+	{
+		Description = FString::Printf(TEXT("%s, detailed description pending"), *GetClass()->GetName());
+	}
+
+	Nickname = GetName();
+}
+
+void UMLAdapterAgentElement::PostInitProperties()
+{
+	// UMLAdapterAgent instance is the only valid outer type
+	check(HasAnyFlags(RF_ClassDefaultObject) || Cast<UMLAdapterAgent>(GetOuter()));
+
+	Super::PostInitProperties();
+
+	if (HasAnyFlags(RF_ClassDefaultObject) == false)
+	{
+		UpdateSpaceDef();
+	}
+}
+
+const UMLAdapterAgent& UMLAdapterAgentElement::GetAgent() const
+{
+	// UMLAdapterAgent instance is the only valid outer type
+	return *CastChecked<UMLAdapterAgent>(GetOuter());
+}
+
+AActor* UMLAdapterAgentElement::GetAvatar() const
+{
+	return GetAgent().GetAvatar();
+}
+
+APawn* UMLAdapterAgentElement::GetPawnAvatar() const
+{
+	AActor* Avatar = GetAvatar();
+	APawn* Pawn = Cast<APawn>(Avatar);
+	AController* Controller = Cast<AController>(Avatar);
+	return Pawn ? Pawn : (Controller ? Controller->GetPawn() : nullptr);
+}
+
+AController* UMLAdapterAgentElement::GetControllerAvatar() const
+{
+	AActor* Avatar = GetAvatar();
+	APawn* Pawn = Cast<APawn>(Avatar);
+	AController* Controller = Cast<AController>(Avatar);
+	return Controller ? Controller : (Pawn ? Pawn->GetController() : nullptr);
+}
+
+bool UMLAdapterAgentElement::GetPawnAndControllerAvatar(APawn*& OutPawn, AController*& OutController) const
+{
+	AActor* Avatar = GetAvatar();
+	APawn* Pawn = Cast<APawn>(Avatar);
+	AController* Controller = Cast<AController>(Avatar);
+	if (Pawn)
+	{
+		Controller = Pawn->GetController();
+	}
+	else if (Controller)
+	{
+		Pawn = Controller->GetPawn();
+	}
+	OutPawn = Pawn;
+	OutController = Controller;
+	
+	return Pawn || Controller;
+}
+
+void UMLAdapterAgentElement::Configure(const TMap<FName, FString>& Params)
+{
+	const FName NAME_Nickname = TEXT("nickname");
+
+	for (auto KeyValue : Params)
+	{
+		if (KeyValue.Key == NAME_Nickname)
+		{
+			Nickname = KeyValue.Value;
+		}
+	}
+}
+
+void UMLAdapterAgentElement::UpdateSpaceDef()
+{
+	SpaceDef = ConstructSpaceDef().ToSharedRef();
+}
+
+#if WITH_GAMEPLAY_DEBUGGER
+#include "GameplayDebuggerCategory.h"
+
+void UMLAdapterAgentElement::DescribeSelfToGameplayDebugger(FGameplayDebuggerCategory& DebuggerCategory) const
+{
+	DebuggerCategory.AddTextLine(FString::Printf(TEXT("\t{yellow}[%d] %s {white}%s"), ElementID
+		, *GetName(), *DebugRuntimeString));
+}
+#endif // WITH_GAMEPLAY_DEBUGGER
diff --git a/Engine/Plugins/AI/MLAdapter/Source/MLAdapter/Private/Agents/MLAdapterAgents.cpp b/Engine/Plugins/AI/MLAdapter/Source/MLAdapter/Private/Agents/MLAdapterAgents.cpp
new file mode 100644
--- /dev/null
+++ b/Engine/Plugins/AI/MLAdapter/Source/MLAdapter/Private/Agents/MLAdapterAgents.cpp
@@ -0,0 +1,459 @@
+// Copyright Epic Games, Inc. All Rights Reserved.
+
+#include "Agents/MLAdapterAgent.h"
+#include "Sensors/MLAdapterSensor.h"
+#include "MLAdapterSession.h"
+#include "MLAdapterSpace.h"
+#include "MLAdapterLibrarian.h"
+#include "GenericPlatform/GenericPlatformMath.h"
+#include "Actuators/MLAdapterActuator_InputKey.h"
+#include "UObject/UObjectGlobals.h"
+#include "GameFramework/PlayerController.h"
+#include "GameFramework/Pawn.h"
+#include "UObject/GarbageCollection.h"
+#include "GameFramework/PlayerState.h"
+
+
+//----------------------------------------------------------------------//
+// FMLAdapterAgentHelpers 
+//----------------------------------------------------------------------//
+namespace FMLAdapterAgentHelpers
+{
+	bool GetAsPawnAndController(AActor* Avatar, AController*& OutController, APawn*& OutPawn)
+	{
+		if (Avatar == nullptr)
+		{
+			return false;
+		}
+
+		APawn* AsPawn = Cast<APawn>(Avatar);
+		if (AsPawn)
+		{
+			OutPawn = AsPawn;
+			OutController = AsPawn->GetController();
+			return true;
+		}
+
+		AController* AsController = Cast<AController>(Avatar);
+		if (AsController)
+		{
+			OutPawn = AsController->GetPawn();
+			OutController = AsController;
+			return true;
+		}
+
+		return false;
+	}
+}
+
+//----------------------------------------------------------------------//
+// FMLAdapterAgentConfig
+//----------------------------------------------------------------------//
+FMLAdapterParameterMap& FMLAdapterAgentConfig::AddSensor(const FName SensorName, FMLAdapterParameterMap&& Parameters)
+{
+	FMLAdapterParameterMap& Entry = Sensors.Add(SensorName, Parameters);
+	return Entry;
+}
+
+FMLAdapterParameterMap& FMLAdapterAgentConfig::AddActuator(const FName ActuatorName, FMLAdapterParameterMap&& Parameters)
+{
+	FMLAdapterParameterMap& Entry = Actuators.Add(ActuatorName, Parameters);
+	return Entry;
+}
+
+//----------------------------------------------------------------------//
+// UMLAdapterAgent 
+//----------------------------------------------------------------------//
+UMLAdapterAgent::UMLAdapterAgent(const FObjectInitializer& ObjectInitializer)
+	: Super(ObjectInitializer)
+{
+	AgentID = FMLAdapter::InvalidAgentID;
+	bEverHadAvatar = false;
+	bRegisteredForPawnControllerChange = false;
+
+	AgentConfig.AddSensor(TEXT("Camera"));
+	AgentConfig.AddSensor(TEXT("Movement"));
+	AgentConfig.AddSensor(TEXT("Input"));
+	AgentConfig.AddActuator(UMLAdapterActuator_InputKey::StaticClass()->GetFName());
+	AgentConfig.Actuators.Add(UMLAdapterActuator_InputKey::StaticClass()->GetFName(), FMLAdapterParameterMap());
+	AgentConfig.AvatarClass = APlayerController::StaticClass();
+	AgentConfig.bAutoRequestNewAvatarUponClearingPrev = true;
+}
+
+void UMLAdapterAgent::BeginDestroy()
+{
+	ShutDownSensorsAndActuators();
+	// forcing unhooking from all event delegates
+	SetAvatar(nullptr); 
+
+	if (bRegisteredForPawnControllerChange)
+	{
+		UGameInstance* GameInstance = GetSession().GetGameInstance();
+		if (GameInstance)
+		{
+			GameInstance->GetOnPawnControllerChanged().RemoveAll(this);
+			bRegisteredForPawnControllerChange = false;
+		}
+	}
+
+	Super::BeginDestroy();
+}
+
+bool UMLAdapterAgent::RegisterSensor(UMLAdapterSensor& Sensor)
+{
+	ensure(Sensor.IsConfiguredForAgent(*this));
+	Sensors.Add(&Sensor);
+	return true;
+}
+
+void UMLAdapterAgent::OnAvatarDestroyed(AActor* DestroyedActor)
+{
+	if (DestroyedActor == Avatar)
+	{
+		SetAvatar(nullptr);
+
+		if (AgentConfig.bAutoRequestNewAvatarUponClearingPrev)
+		{
+			// note that after this call we might not have the avatar just yet 
+			// since the world might not have any. The Session will make sure to 
+			// assign us one as soon as one's ready.
+			GetSession().RequestAvatarForAgent(*this);
+		}
+	}
+}
+
+void UMLAdapterAgent::OnPawnChanged(APawn* NewPawn, AController* InController)
+{
+	ensure(Controller == InController);
+	if (Controller != InController)
+	{
+		// this came from a different controller we somehow missed unbinding from. Ignore.
+		return;
+	}
+
+	if (Pawn == NewPawn)
+	{
+		return;
+	}
+
+	// let every sense that requires a pawn know that the pawn changed 
+	for (UMLAdapterSensor* Sensor : Sensors)
+	{
+		Sensor->OnPawnChanged(Pawn, NewPawn);
+	}
+
+	Pawn = NewPawn;
+}
+
+void UMLAdapterAgent::OnPawnControllerChanged(APawn* InPawn, AController* InController)
+{
+	if (InPawn == Pawn)
+	{
+		if (Pawn == Avatar)
+		{
+			Controller = InController;
+		}
+		// if controller is the main avatar we might have just lost our pawn
+		else if (Controller && Controller != InController && Controller == Avatar)
+		{
+			OnPawnChanged(Controller->GetPawn() == InPawn ? nullptr : Controller->GetPawn(), Controller);
+		}
+	}
+}
+
+void UMLAdapterAgent::Sense(const float DeltaTime)
+{
+	for (UMLAdapterSensor* Sensor : Sensors)
+	{
+		// not that due to the system's design Sensor won't be null
+		Sensor->Sense(DeltaTime);
+	}
+}
+
+void UMLAdapterAgent::Act(const float DeltaTime)
+{
+	for (UMLAdapterActuator* Actuator : Actuators)
+	{
+		Actuator->Act(DeltaTime);
+	}
+}
+
+void UMLAdapterAgent::DigestActions(FMLAdapterMemoryReader& ValueStream)
+{
+	for (UMLAdapterActuator* Actuator : Actuators)
+	{
+		Actuator->DigestInputData(ValueStream);
+	}
+}
+
+void UMLAdapterAgent::GetObservations(FMLAdapterMemoryWriter& Ar)
+{
+	for (UMLAdapterSensor* Sensor : Sensors)
+	{
+		Sensor->GetObservations(Ar);
+	}
+}
+
+void UMLAdapterAgent::ShutDownSensorsAndActuators()
+{	
+	for (UMLAdapterActuator* Actuator : Actuators)
+	{
+		if (Actuator)
+		{
+			Actuator->Shutdown();
+		}
+	}
+	Actuators.Reset();
+
+	for (UMLAdapterSensor* Sensor : Sensors)
+	{
+		if (Sensor)
+		{
+			Sensor->Shutdown();
+		}
+	}
+	Sensors.Reset();
+}
+
+void UMLAdapterAgent::Configure(const FMLAdapterAgentConfig& NewConfig)
+{
+	ShutDownSensorsAndActuators();
+
+	TSubclassOf<AActor> PreviousAvatarClass = AgentConfig.AvatarClass;
+	AgentConfig = NewConfig;
+
+	for (const TTuple<FName, FMLAdapterParameterMap>& KeyValue : AgentConfig.Actuators)
+	{
+		UClass* ResultClass = FMLAdapterLibrarian::Get().FindActuatorClass(KeyValue.Key);
+		if (ResultClass)
+		{
+			UMLAdapterActuator* NewActuator = FMLAdapter::NewObject<UMLAdapterActuator>(this, ResultClass);
+			check(NewActuator);
+			NewActuator->SetNickname(KeyValue.Key.ToString());
+			NewActuator->Configure(KeyValue.Value.Params);
+			Actuators.Add(NewActuator);
+		}
+	}
+	Actuators.Sort(FAgentElementSort());
+
+	for (const TTuple<FName, FMLAdapterParameterMap>& KeyValue : AgentConfig.Sensors)
+	{
+		UClass* ResultClass = FMLAdapterLibrarian::Get().FindSensorClass(KeyValue.Key);
+		if (ResultClass)
+		{
+			UMLAdapterSensor* NewSensor = FMLAdapter::NewObject<UMLAdapterSensor>(this, ResultClass);
+			check(NewSensor);
+			NewSensor->SetNickname(KeyValue.Key.ToString());
+			NewSensor->Configure(KeyValue.Value.Params);
+			Sensors.Add(NewSensor);
+		}
+	}
+	Sensors.Sort(FAgentElementSort());
+
+	if (NewConfig.AvatarClassName != NAME_None)
+	{
+		AgentConfig.AvatarClass = FindObject<UClass>(ANY_PACKAGE, *NewConfig.AvatarClassName.ToString());
+	}
+	if (!AgentConfig.AvatarClass)
+	{
+		AgentConfig.AvatarClass = AActor::StaticClass();
+	}
+
+	ensure(AgentConfig.AvatarClass || Avatar);
+
+	if (AgentConfig.AvatarClass && (Avatar == nullptr || IsSuitableAvatar(*Avatar) == false))
+	{
+		SetAvatar(nullptr);
+
+		// if avatar class changed make sure the following RequestAvatarForAgent actually tries to find an avatar
+		// rather than just ignoring the request due to the agent already being in AwaitingAvatar
+		const bool bForceUpdate = (AgentConfig.AvatarClass != PreviousAvatarClass);
+
+		// note that after this call we might not have the avatar just yet 
+		// since the world might not have any. The Session will make sure to 
+		// assign us one as soon as one's ready.
+		GetSession().RequestAvatarForAgent(*this, /*World=*/nullptr, bForceUpdate);
+	}
+	else if (ensure(Avatar))
+	{
+		// newly created actuators and sensors might need the information about 
+		// the current avatar
+		for (UMLAdapterSensor* Sensor : Sensors)
+		{
+			Sensor->OnAvatarSet(Avatar);
+		}
+		for (UMLAdapterActuator* Actuator : Actuators)
+		{
+			Actuator->OnAvatarSet(Avatar);
+		}
+	}
+}
+
+UMLAdapterSession& UMLAdapterAgent::GetSession()
+{
+	UMLAdapterSession* Session = CastChecked<UMLAdapterSession>(GetOuter());
+	return *Session;
+}
+
+void UMLAdapterAgent::GetActionSpaceDescription(FMLAdapterSpaceDescription& OutSpaceDesc) const
+{
+	FMLAdapterDescription ElementDesc;
+
+	for (UMLAdapterActuator* Actuator : Actuators)
+	{
+		if (Actuator)
+		{
+			ElementDesc.Reset();
+			ElementDesc.Add(Actuator->GetSpaceDef());
+			OutSpaceDesc.Add(Actuator->GetNickname(), ElementDesc);
+		}
+	}
+}
+
+void UMLAdapterAgent::GetObservationSpaceDescription(FMLAdapterSpaceDescription& OutSpaceDesc) const
+{
+	FMLAdapterDescription ElementDesc;
+
+	for (UMLAdapterSensor* Sensor : Sensors)
+	{
+		if (Sensor)
+		{
+			ElementDesc.Reset();
+			ElementDesc.Add(Sensor->GetSpaceDef());
+			OutSpaceDesc.Add(Sensor->GetNickname(), ElementDesc);
+		}
+	}
+}
+
+bool UMLAdapterAgent::IsSuitableAvatar(AActor& InAvatar) const
+{	
+	return AgentConfig.bAvatarClassExact 
+		? InAvatar.GetClass() == AgentConfig.AvatarClass.Get()
+		: InAvatar.IsA(AgentConfig.AvatarClass);
+}
+
+void UMLAdapterAgent::SetAvatar(AActor* InAvatar)
+{
+	if (InAvatar == Avatar)
+	{
+		// do nothing on purpose
+		return;
+	}
+
+	if (InAvatar != nullptr && IsSuitableAvatar(*InAvatar) == false)
+	{
+		UE_LOG(LogUnrealEditorMLAdapter, Log, TEXT("SetAvatar was called for agent %u but %s is not a valid avatar (required avatar class %s)")
+			, AgentID, *InAvatar->GetName(), *GetNameSafe(AgentConfig.AvatarClass));
+		return;
+	}
+
+	AActor* PrevAvatar = Avatar;
+	AController* PrevController = Controller;
+	APawn* PrevPawn = Pawn;
+	if (Avatar)
+	{
+		Avatar->OnDestroyed.RemoveAll(this);
+		Avatar = nullptr;
+	}
+
+	if (InAvatar == nullptr)
+	{
+		Controller = nullptr;
+		Pawn = nullptr;
+	}
+	else
+	{
+		bEverHadAvatar = true;
+		Avatar = InAvatar;
+
+		Pawn = nullptr;
+		Controller = nullptr;
+		FMLAdapterAgentHelpers::GetAsPawnAndController(Avatar, Controller, Pawn);
+
+		if (Avatar)
+		{
+			Avatar->OnDestroyed.AddDynamic(this, &UMLAdapterAgent::OnAvatarDestroyed);
+		}
+	}
+
+	// actuators and sensors might need the information that the avatar changed
+	for (UMLAdapterSensor* Sensor : Sensors)
+	{
+		Sensor->OnAvatarSet(Avatar);
+	}
+	for (UMLAdapterActuator* Actuator : Actuators)
+	{
+		Actuator->OnAvatarSet(Avatar);
+	}
+
+	// unregister from unused notifies
+
+	if (Controller != PrevController)
+	{
+		if (PrevController && PrevController == PrevAvatar)
+		{
+			PrevController->GetOnNewPawnNotifier().RemoveAll(this);
+		}
+		// when the controller is the main avatar
+		if (Controller != nullptr && (Avatar == Controller))
+		{
+			Controller->GetOnNewPawnNotifier().AddUObject(this, &UMLAdapterAgent::OnPawnChanged, Controller);
+		}
+	}
+
+	if ((Controller != nullptr || Pawn != nullptr) && bRegisteredForPawnControllerChange == false)
+	{
+		UGameInstance* GameInstance = GetSession().GetGameInstance();
+		if (GameInstance)
+		{
+			GameInstance->GetOnPawnControllerChanged().AddDynamic(this, &UMLAdapterAgent::OnPawnControllerChanged);
+			bRegisteredForPawnControllerChange = true;
+		}
+	}
+}
+
+float UMLAdapterAgent::GetReward() const
+{
+	if (Avatar == nullptr)
+	{
+		return 0.f;
+	}
+	FGCScopeGuard GCScopeGuard;
+	AController* AsController = FMLAdapter::ActorToController(*Avatar);
+	return AsController && AsController->PlayerState ? AsController->PlayerState->GetScore() : 0.f;
+}
+
+bool UMLAdapterAgent::IsDone() const
+{
+	return AgentConfig.bAutoRequestNewAvatarUponClearingPrev == false 
+		&& Avatar == nullptr
+		&& bEverHadAvatar == true;
+}
+
+#if WITH_GAMEPLAY_DEBUGGER
+#include "GameplayDebuggerCategory.h"
+
+void UMLAdapterAgent::DescribeSelfToGameplayDebugger(FGameplayDebuggerCategory& DebuggerCategory) const
+{
+	DebuggerCategory.AddTextLine(FString::Printf(TEXT("{green}ID {white}%u\n{green}Avatar {white}%s")
+		, AgentID, *GetNameSafe(Avatar)));
+
+	DebuggerCategory.AddTextLine(TEXT("{green}Sensors:"));
+	for (UMLAdapterSensor* Sensor : Sensors)
+	{
+		if (Sensor)
+		{
+			Sensor->DescribeSelfToGameplayDebugger(DebuggerCategory);
+		}
+	}
+
+	DebuggerCategory.AddTextLine(TEXT("{green}Actuators:"));
+	for (UMLAdapterActuator* Actuator : Actuators)
+	{
+		if (Actuator)
+		{
+			Actuator->DescribeSelfToGameplayDebugger(DebuggerCategory);
+		}
+	}
+}
+#endif // WITH_GAMEPLAY_DEBUGGER
\ No newline at end of file
diff --git a/Engine/Plugins/AI/MLAdapter/Source/MLAdapter/Public/Agents/MLAdapterAgentElement.h b/Engine/Plugins/AI/MLAdapter/Source/MLAdapter/Public/Agents/MLAdapterAgentElement.h
new file mode 100644
--- /dev/null
+++ b/Engine/Plugins/AI/MLAdapter/Source/MLAdapter/Public/Agents/MLAdapterAgentElement.h
@@ -0,0 +1,99 @@
+// Copyright Epic Games, Inc. All Rights Reserved.
+
+#pragma once 
+
+#include "CoreMinimal.h"
+#include "UObject/ObjectMacros.h"
+#include "UObject/UObjectGlobals.h"
+#include "UObject/Object.h"
+#include "MLAdapterSpace.h"
+#include "MLAdapterAgentElement.generated.h"
+
+
+class UMLAdapterAgent;
+struct FMLAdapterDescription;
+
+UCLASS(abstract)
+class MLADAPTER_API UMLAdapterAgentElement : public UObject
+{
+	GENERATED_BODY()
+public:
+	UMLAdapterAgentElement(const FObjectInitializer& ObjectInitializer = FObjectInitializer::Get());
+	virtual void PostInitProperties() override;
+
+	const UMLAdapterAgent& GetAgent() const;
+	uint32 GetElementID() const { return ElementID; }
+	const FString& GetNickname() const { return Nickname; }
+
+	/** @return owning agent's avatar */
+	AActor* GetAvatar() const;
+
+	/** @return the pawn associated with the owning agent. If owning agent's avatar 
+	 *	is a pawn then that gets retrieved, if not the function will check if 
+	 *	it's a controller and if so retrieve its pawn */
+	APawn* GetPawnAvatar() const;
+
+	/** @return the controller associated with the owning agent. If owning agent's 
+	 *	avatar is a controller then that gets retrieved, if not the function will 
+	 *	check if it's a pawn and if so retrieve its controller */
+	AController* GetControllerAvatar() const;
+
+	/** Fetches both the pawn and the controller associated with the current agent.
+	 *	It's like both calling @see GetPawnAvatar and @see GetControllerAvatar 
+	 *	@return true if at least one of the fetched pair is non-null */
+	bool GetPawnAndControllerAvatar(APawn*& OutPawn, AController*& OutController) const;
+
+	virtual FString GetDescription() const { return Description; }
+
+	/** Called before object's destruction. Can be called as part of new agent 
+	 *	config application when old actuator get destroyed */
+	virtual void Shutdown() {}
+	
+	void SetNickname(const FString& NewNickname) { Nickname = NewNickname; }
+	virtual void Configure(const TMap<FName, FString>& Params);
+	virtual void OnAvatarSet(AActor* Avatar) {}
+
+	virtual void UpdateSpaceDef();
+	const FMLAdapter::FSpace& GetSpaceDef() const { return *SpaceDef; }
+
+#if WITH_GAMEPLAY_DEBUGGER
+	virtual void DescribeSelfToGameplayDebugger(class FGameplayDebuggerCategory& DebuggerCategory) const;
+#endif // WITH_GAMEPLAY_DEBUGGER
+
+protected:
+	virtual TSharedPtr<FMLAdapter::FSpace> ConstructSpaceDef() const PURE_VIRTUAL(UMLAdapterAgentElement::ConstructSpaceDef, return MakeShareable(new FMLAdapter::FSpace_Dummy()); );
+
+protected:
+	// can be queried by remote clients 
+	FString Description;
+
+	TSharedRef<FMLAdapter::FSpace> SpaceDef;
+
+	/** @note this is not a common counter, meaning Sensors and Actuators (for 
+	 *	example) track the ID separately */
+	UPROPERTY()
+	uint32 ElementID;
+
+	/** User-configured name for this element, mostly for debugging purposes but 
+	 *	comes in handy when fetching observation/action spaces descriptions.
+	 *	Defaults to UE4 instance name*/
+	UPROPERTY()
+	FString Nickname; 
+
+#if WITH_GAMEPLAY_DEBUGGER
+	// displayed in debugging tools and logging
+	mutable FString DebugRuntimeString;
+#endif // WITH_GAMEPLAY_DEBUGGER
+};
+
+struct FAgentElementSort
+{
+	bool operator()(const UMLAdapterAgentElement* A, const UMLAdapterAgentElement* B) const
+	{
+		return A && (!B || (A->GetElementID() < B->GetElementID()));
+	}
+	bool operator()(const UMLAdapterAgentElement& A, const UMLAdapterAgentElement& B) const
+	{
+		return (A.GetElementID() < B.GetElementID());
+	}
+};
