UE4ML->MLAdapter rename follow up fixes

diff --git a/Engine/Plugins/AI/MLAdapter/Source/MLAdapter/Private/MLAdapterManager_Client.cpp b/Engine/Plugins/AI/MLAdapter/Source/MLAdapter/Private/MLAdapterManager_Client.cpp
--- a/Engine/Plugins/AI/MLAdapter/Source/MLAdapter/Private/MLAdapterManager_Client.cpp
+++ b/Engine/Plugins/AI/MLAdapter/Source/MLAdapter/Private/MLAdapterManager_Client.cpp
@@ -64,7 +64,7 @@ void UMLAdapterManager::AddCommonFunctions(FRPCServer& Server)
 	Server.bind("exit", []() {
 		FPlatformMisc::RequestExit(/*bForce=*/false);
 	});
-	Librarian.AddRPCFunctionDescription(TEXT("exit"), TEXT("(), Closes the UE4 instance."));
+	Librarian.AddRPCFunctionDescription(TEXT("exit"), TEXT("(), Closes the UnrealEngine instance."));
 	
 	Server.bind("batch_is_finished", [this](std::vector<FMLAdapter::FAgentID> AgentIDs) {
 		std::vector<bool> Results;
diff --git a/Engine/Plugins/AI/MLAdapter/Source/MLAdapter/Public/Agents/MLAdapterAgentElement.h b/Engine/Plugins/AI/MLAdapter/Source/MLAdapter/Public/Agents/MLAdapterAgentElement.h
--- a/Engine/Plugins/AI/MLAdapter/Source/MLAdapter/Public/Agents/MLAdapterAgentElement.h
+++ b/Engine/Plugins/AI/MLAdapter/Source/MLAdapter/Public/Agents/MLAdapterAgentElement.h
@@ -76,7 +76,7 @@ protected:
 
 	/** User-configured name for this element, mostly for debugging purposes but 
 	 *	comes in handy when fetching observation/action spaces descriptions.
-	 *	Defaults to UE4 instance name*/
+	 *	Defaults to UnrealEngine instance name*/
 	UPROPERTY()
 	FString Nickname; 
 
diff --git a/Engine/Plugins/AI/MLAdapter/Source/MLAdapter/Public/MLAdapterManager.h b/Engine/Plugins/AI/MLAdapter/Source/MLAdapter/Public/MLAdapterManager.h
--- a/Engine/Plugins/AI/MLAdapter/Source/MLAdapter/Public/MLAdapterManager.h
+++ b/Engine/Plugins/AI/MLAdapter/Source/MLAdapter/Public/MLAdapterManager.h
@@ -77,7 +77,7 @@ public:
 	virtual bool IsRunning() const;
 
 	virtual void ConfigureAsServer(FRPCServer& Server);
-	/** "Client" in this context means UE4 game client, not RPC client */
+	/** "Client" in this context means UnrealEngine game client, not RPC client */
 	virtual void ConfigureAsClient(FRPCServer& Server);
 	/** Essentially calls both the server and client versions */
 	virtual void ConfigureAsStandalone(FRPCServer& Server);
diff --git a/Engine/Plugins/AI/MLAdapter/Source/MLAdapter/Public/MLAdapterSession.h b/Engine/Plugins/AI/MLAdapter/Source/MLAdapter/Public/MLAdapterSession.h
--- a/Engine/Plugins/AI/MLAdapter/Source/MLAdapter/Public/MLAdapterSession.h
+++ b/Engine/Plugins/AI/MLAdapter/Source/MLAdapter/Public/MLAdapterSession.h
@@ -65,7 +65,7 @@ public:
 	virtual void ConfigureAsServer();
 	/** This is where Session can add Client-side-specific functions by calling 
 	 *	UMLAdapterManager::Get().AddClientFunctionBind. "Client" in this context means 
-	 *	UE4 game client, not RPC client */
+	 *	UnrealEngine game client, not RPC client */
 	virtual void ConfigureAsClient();
 
 	/** using FMLAdapter::InvalidAgentID for AgentID will reset all agents */
diff --git a/Engine/Plugins/AI/MLAdapter/Source/MLAdapter/Public/RPCWrapper/MsgPack.h b/Engine/Plugins/AI/MLAdapter/Source/MLAdapter/Public/RPCWrapper/MsgPack.h
--- a/Engine/Plugins/AI/MLAdapter/Source/MLAdapter/Public/RPCWrapper/MsgPack.h
+++ b/Engine/Plugins/AI/MLAdapter/Source/MLAdapter/Public/RPCWrapper/MsgPack.h
@@ -5,14 +5,14 @@
 #if WITH_RPCLIB
 #include "rpclib_includes.h"
 
-#define UE4_VIRTUAL_MSGPACK_DEFINE(...)                                                \
+#define UE_VIRTUAL_MSGPACK_DEFINE(...)                                                \
     virtual void msgpack_object(RPCLIB_MSGPACK::object *o, clmdep_msgpack::zone &z)    \
         const {                                                                \
         clmdep_msgpack::type::make_define_array(__VA_ARGS__)                   \
             .msgpack_object(o, z);                                             \
     }
 
-#define UE4_MSGPACK_DEFINE(...)                                                \
+#define UE_MSGPACK_DEFINE(...)                                                \
     void msgpack_object(RPCLIB_MSGPACK::object *o, clmdep_msgpack::zone &z)    \
         const {                                                                \
         clmdep_msgpack::type::make_define_array(__VA_ARGS__)                   \
@@ -21,8 +21,8 @@
 
 #else // WITH_RPCLIB
 
-#define UE4_MSGPACK_DEFINE(...)
-#define UE4_VIRTUAL_MSGPACK_DEFINE(...)
+#define UE_MSGPACK_DEFINE(...)
+#define UE_VIRTUAL_MSGPACK_DEFINE(...)
 #define MSGPACK_DEFINE_ARRAY(...)
 
 #endif // WITH_RPCLIB
@@ -30,7 +30,7 @@
 // helper macro for just binding a simple function. Not recommended for elaborate
 // lambdas since using this macro interferes with debugging (unable to break in 
 // the lambdas body).
-#define UE4_RPC_BIND(FuncName, FuncBody) TEXT(FuncName), [this](FRPCServer& Serv) { Serv.bind(FuncName, FuncBody); }
+#define UE_RPC_BIND(FuncName, FuncBody) TEXT(FuncName), [this](FRPCServer& Serv) { Serv.bind(FuncName, FuncBody); }
 #define FSTRING_TO_STD(a) (std::string(TCHAR_TO_UTF8(*a)))
 
 #include <vector>
