Fixed var shadowing in MLAdapterManager_Client.cpp

horde issue #62517

diff --git a/Engine/Plugins/AI/MLAdapter/Source/MLAdapter/Private/MLAdapterManager_Client.cpp b/Engine/Plugins/AI/MLAdapter/Source/MLAdapter/Private/MLAdapterManager_Client.cpp
--- a/Engine/Plugins/AI/MLAdapter/Source/MLAdapter/Private/MLAdapterManager_Client.cpp
+++ b/Engine/Plugins/AI/MLAdapter/Source/MLAdapter/Private/MLAdapterManager_Client.cpp
@@ -77,10 +77,10 @@ void UMLAdapterManager::AddCommonFunctions(FRPCServer& Server)
 		}
 		else
 		{
-			UMLAdapterSession& Session = GetSession();
+			UMLAdapterSession& SessionInstance = GetSession();
 			for (const FMLAdapter::FAgentID& AgentID : AgentIDs)
 			{
-				UMLAdapterAgent* Agent = Session.GetAgent(AgentID);
+				UMLAdapterAgent* Agent = SessionInstance.GetAgent(AgentID);
 				Results.push_back(Agent == nullptr || Agent->IsDone());
 			}
 		}
