MLAdapter: Minor changes (spelling errors) and adding Think to Session

#rnx

diff --git a/Engine/Plugins/AI/MLAdapter/Source/MLAdapter/Private/Agents/MLAdapterAgents.cpp b/Engine/Plugins/AI/MLAdapter/Source/MLAdapter/Private/Agents/MLAdapterAgents.cpp
--- a/Engine/Plugins/AI/MLAdapter/Source/MLAdapter/Private/Agents/MLAdapterAgents.cpp
+++ b/Engine/Plugins/AI/MLAdapter/Source/MLAdapter/Private/Agents/MLAdapterAgents.cpp
@@ -170,6 +170,11 @@ void UMLAdapterAgent::Sense(const float DeltaTime)
 	}
 }
 
+void UMLAdapterAgent::Think(const float DeltaTime)
+{
+	// Can be implemented in child class
+}
+
 void UMLAdapterAgent::Act(const float DeltaTime)
 {
 	for (UMLAdapterActuator* Actuator : Actuators)
diff --git a/Engine/Plugins/AI/MLAdapter/Source/MLAdapter/Private/MLAdapterManager.cpp b/Engine/Plugins/AI/MLAdapter/Source/MLAdapter/Private/MLAdapterManager.cpp
--- a/Engine/Plugins/AI/MLAdapter/Source/MLAdapter/Private/MLAdapterManager.cpp
+++ b/Engine/Plugins/AI/MLAdapter/Source/MLAdapter/Private/MLAdapterManager.cpp
@@ -326,7 +326,7 @@ void UMLAdapterManager::BindToDelegates()
 	// Commented out possible other useful delegates
 	//	FCoreDelegates::GameStateClassChanged;
 	//	FCoreDelegates::ConfigReadyForUse;
-	//	FWorldDelegates::OnPostWorldCreation;
+	//  FWorldDelegates::OnPostWorldCreation;
 	//	FWorldDelegates::OnPreWorldInitialization; 
 
 	FWorldDelegates::OnPostWorldInitialization.AddUObject(this, &UMLAdapterManager::OnPostWorldInit);
@@ -381,11 +381,11 @@ void UMLAdapterManager::OnPostWorldInit(UWorld* World, const UWorld::Initializat
 #endif // WITH_EDITOR
 
 		LastActiveWorld = World;
-		if ((RequestedFunctionMode == EMLAdapterServerMode::Invalid || RequestedFunctionMode == EMLAdapterServerMode::AutoDetect) 
+		if ((RequestedFunctionMode == EMLAdapterServerMode::Invalid || RequestedFunctionMode == EMLAdapterServerMode::AutoDetect)
 			&& GetServerModeForWorld(*World) != CurrentFunctionMode)
 		{
 			// restart the RPC server. Note that this will kick all the currently connected agents
-			uint16 Port = UMLAdapterSettings::GetDefaultRPCServerPort();			
+			uint16 Port = UMLAdapterSettings::GetDefaultRPCServerPort();
 			FParse::Value(FCommandLine::Get(), TEXT("MLAdapterPort="), Port);
 			StartServer(Port, GetServerModeForWorld(*World), CurrentServerThreads);
 		}
@@ -394,7 +394,7 @@ void UMLAdapterManager::OnPostWorldInit(UWorld* World, const UWorld::Initializat
 		{
 			GetSession().OnPostWorldInit(*World);
 		}
-	}	
+	}
 }
 
 void UMLAdapterManager::OnWorldCleanup(UWorld* World, bool bSessionEnded, bool bCleanupResources)
@@ -409,7 +409,7 @@ void UMLAdapterManager::OnWorldCleanup(UWorld* World, bool bSessionEnded, bool b
 
 		if (HasSession())
 		{
-			GetSession().OnWorldCleanup(*World, bSessionEnded, bCleanupResources);
+			CloseSession(GetSession());
 		}
 	}
 }
diff --git a/Engine/Plugins/AI/MLAdapter/Source/MLAdapter/Private/MLAdapterSession.cpp b/Engine/Plugins/AI/MLAdapter/Source/MLAdapter/Private/MLAdapterSession.cpp
--- a/Engine/Plugins/AI/MLAdapter/Source/MLAdapter/Private/MLAdapterSession.cpp
+++ b/Engine/Plugins/AI/MLAdapter/Source/MLAdapter/Private/MLAdapterSession.cpp
@@ -225,7 +225,7 @@ void UMLAdapterSession::Close()
 {
 	bActive = false;
 
-	// 'destroy' agents by clearing the asyc flag and letting the GC clean them
+	// 'destroy' agents by clearing the async flag and letting the GC clean them
 	for (UMLAdapterAgent* Agent : Agents)
 	{
 		Agent->ClearInternalFlags(EInternalObjectFlags::Async);
@@ -251,7 +251,7 @@ void UMLAdapterSession::Tick(float DeltaTime)
 
 	for (UMLAdapterAgent* Agent : Agents)
 	{
-		//Agent->Think(DeltaTime);
+		Agent->Think(DeltaTime);
 	}
 
 	for (UMLAdapterAgent* Agent : Agents)
diff --git a/Engine/Plugins/AI/MLAdapter/Source/MLAdapter/Private/MLAdapterSettings.cpp b/Engine/Plugins/AI/MLAdapter/Source/MLAdapter/Private/MLAdapterSettings.cpp
--- a/Engine/Plugins/AI/MLAdapter/Source/MLAdapter/Private/MLAdapterSettings.cpp
+++ b/Engine/Plugins/AI/MLAdapter/Source/MLAdapter/Private/MLAdapterSettings.cpp
@@ -7,7 +7,7 @@
 UMLAdapterSettings::UMLAdapterSettings(const FObjectInitializer& ObjectInitializer)
 	: Super(ObjectInitializer)
 {
-	DefautAgentClass = UMLAdapterAgent::StaticClass();
+	DefaultAgentClass = UMLAdapterAgent::StaticClass();
 	ManagerClass = UMLAdapterManager::StaticClass();
 	SessionClass = UMLAdapterSession::StaticClass();
 }
@@ -28,7 +28,7 @@ TSubclassOf<UMLAdapterSession> UMLAdapterSettings::GetSessionClass()
 
 TSubclassOf<UMLAdapterAgent> UMLAdapterSettings::GetAgentClass()
 {
-	const FSoftClassPath LocalClassName = GET_CONFIG_VALUE(DefautAgentClass);
+	const FSoftClassPath LocalClassName = GET_CONFIG_VALUE(DefaultAgentClass);
 	TSubclassOf<UMLAdapterAgent> LocalClass = LocalClassName.ResolveClass();
 	return LocalClass;
 }
diff --git a/Engine/Plugins/AI/MLAdapter/Source/MLAdapter/Public/Agents/MLAdapterAgent.h b/Engine/Plugins/AI/MLAdapter/Source/MLAdapter/Public/Agents/MLAdapterAgent.h
--- a/Engine/Plugins/AI/MLAdapter/Source/MLAdapter/Public/Agents/MLAdapterAgent.h
+++ b/Engine/Plugins/AI/MLAdapter/Source/MLAdapter/Public/Agents/MLAdapterAgent.h
@@ -13,7 +13,6 @@ class UMLAdapterAction;
 class UMLAdapterSensor;
 class AController;
 class APawn;
-class UMLAdapterBrain;
 class UMLAdapterSession;
 class UMLAdapterActuator;
 struct FMLAdapterSpaceDescription;
@@ -73,9 +72,13 @@ public:
 	virtual void BeginDestroy() override;
 
 	virtual bool RegisterSensor(UMLAdapterSensor& Sensor);
+	
 	/** Updates all the senses that are configured as 'IsPolling'*/
 	virtual void Sense(const float DeltaTime);
 	
+	// Decide what action to take based on the current observations
+	virtual void Think(const float DeltaTime);
+
 	// trigger all of the agent's 
 	virtual void Act(const float DeltaTime);
 	
diff --git a/Engine/Plugins/AI/MLAdapter/Source/MLAdapter/Public/MLAdapterJson.h b/Engine/Plugins/AI/MLAdapter/Source/MLAdapter/Public/MLAdapterJson.h
--- a/Engine/Plugins/AI/MLAdapter/Source/MLAdapter/Public/MLAdapterJson.h
+++ b/Engine/Plugins/AI/MLAdapter/Source/MLAdapter/Public/MLAdapterJson.h
@@ -38,9 +38,9 @@ namespace FMLAdapter
 	}
 
 	template<typename StructType>
-	void JsonStringToStruct(const FString& JsonString, StructType& OutStruct)
+	bool JsonStringToStruct(const FString& JsonString, StructType& OutStruct)
 	{
-		FJsonObjectConverter::JsonObjectStringToUStruct(JsonString, &OutStruct, 0, 0);
+		return FJsonObjectConverter::JsonObjectStringToUStruct(JsonString, &OutStruct, 0, 0);
 	}
 
 	struct IJsonable
diff --git a/Engine/Plugins/AI/MLAdapter/Source/MLAdapter/Public/MLAdapterManager.h b/Engine/Plugins/AI/MLAdapter/Source/MLAdapter/Public/MLAdapterManager.h
--- a/Engine/Plugins/AI/MLAdapter/Source/MLAdapter/Public/MLAdapterManager.h
+++ b/Engine/Plugins/AI/MLAdapter/Source/MLAdapter/Public/MLAdapterManager.h
@@ -92,7 +92,7 @@ public:
 	virtual void SetSession(UMLAdapterSession* NewSession);
 	virtual void CloseSession(UMLAdapterSession& InSession);
 
-	/** Returns current sesison. If one doesn't exist, it gets created. */
+	/** Returns current session. If one doesn't exist, it gets created. */
 	virtual UMLAdapterSession& GetSession();
 	bool HasSession() const { return Session && (Session->IsPendingKillOrUnreachable() == false); }
 
diff --git a/Engine/Plugins/AI/MLAdapter/Source/MLAdapter/Public/MLAdapterSession.h b/Engine/Plugins/AI/MLAdapter/Source/MLAdapter/Public/MLAdapterSession.h
--- a/Engine/Plugins/AI/MLAdapter/Source/MLAdapter/Public/MLAdapterSession.h
+++ b/Engine/Plugins/AI/MLAdapter/Source/MLAdapter/Public/MLAdapterSession.h
@@ -40,6 +40,7 @@ public:
 
 	virtual UWorld* GetWorld() const override { return CachedWorld; }
 	virtual void PostInitProperties() override;
+
 	virtual void BeginDestroy() override;
 
 	// @todo this needs further consideration.
@@ -63,6 +64,7 @@ public:
 	/** This is where Session can add Authority-side-specific functions by calling 
 	 *	UMLAdapterManager::Get().AddServerFunctionBind */
 	virtual void ConfigureAsServer();
+
 	/** This is where Session can add Client-side-specific functions by calling 
 	 *	UMLAdapterManager::Get().AddClientFunctionBind. "Client" in this context means 
 	 *	UnrealEngine game client, not RPC client */
@@ -89,7 +91,7 @@ public:
 
 	/** @return next valid agent ID. Note that the return value might be equal to 
 	 *	ReferenceAgentID if there's only one agent. Will be FMLAdapter::InvalidAgentId 
-	 *	if no agents registerd */
+	 *	if no agents registered */
 	FMLAdapter::FAgentID GetNextAgentID(FMLAdapter::FAgentID ReferenceAgentID) const;
 	UMLAdapterAgent* GetAgent(FMLAdapter::FAgentID AgentID);
 
diff --git a/Engine/Plugins/AI/MLAdapter/Source/MLAdapter/Public/MLAdapterSettings.h b/Engine/Plugins/AI/MLAdapter/Source/MLAdapter/Public/MLAdapterSettings.h
--- a/Engine/Plugins/AI/MLAdapter/Source/MLAdapter/Public/MLAdapterSettings.h
+++ b/Engine/Plugins/AI/MLAdapter/Source/MLAdapter/Public/MLAdapterSettings.h
@@ -42,7 +42,7 @@ protected:
 	FSoftClassPath SessionClass;
 	
 	UPROPERTY(EditAnywhere, config, Category = MLAdapter, meta = (MetaClass = "MLAdapterAgent"))
-	FSoftClassPath DefautAgentClass;
+	FSoftClassPath DefaultAgentClass;
 
 	UPROPERTY(EditAnywhere, config, Category = MLAdapter)
 	uint16 DefaultRPCServerPort = 15151;
diff --git a/Engine/Plugins/AI/MLAdapter/Source/MLAdapter/Public/MLAdapterTypes.h b/Engine/Plugins/AI/MLAdapter/Source/MLAdapter/Public/MLAdapterTypes.h
--- a/Engine/Plugins/AI/MLAdapter/Source/MLAdapter/Public/MLAdapterTypes.h
+++ b/Engine/Plugins/AI/MLAdapter/Source/MLAdapter/Public/MLAdapterTypes.h
@@ -92,7 +92,7 @@ struct FMLAdapterMemoryWriter : FMemoryWriter
 
 struct FMLAdapterMemoryReader : FMemoryReader
 {
-	FMLAdapterMemoryReader(TArray<uint8>& InBytes, bool bIsPersistent = false)
+	FMLAdapterMemoryReader(const TArray<uint8>& InBytes, bool bIsPersistent = false)
 		: FMemoryReader(InBytes, bIsPersistent)
 	{}
 };
