MLAdapter Plugin: settings can be changed without needed to reboot the editor

#rnx

diff --git a/Engine/Plugins/AI/MLAdapter/Source/MLAdapter/Private/MLAdapterManager.cpp b/Engine/Plugins/AI/MLAdapter/Source/MLAdapter/Private/MLAdapterManager.cpp
--- a/Engine/Plugins/AI/MLAdapter/Source/MLAdapter/Private/MLAdapterManager.cpp
+++ b/Engine/Plugins/AI/MLAdapter/Source/MLAdapter/Private/MLAdapterManager.cpp
@@ -50,39 +50,43 @@ namespace
 	}
 }
 
-namespace
+void UMLAdapterManager::RecreateManagerInstance()
 {
-	struct FManagerBootloader
+#if WITH_EDITOR
+	if (UMLAdapterManager::IsReady())
 	{
-		FManagerBootloader()
-		{
-			FCoreDelegates::OnPostEngineInit.AddLambda([this]()
-			{
-				OnPostEngineInit();
-			});
-		}
-
-		void OnPostEngineInit()
-		{
-			// create the manager instance
-			TSubclassOf<UMLAdapterManager> SettingsManagerClass = UMLAdapterSettings::GetManagerClass().Get();
-			UClass* Class = SettingsManagerClass
-				? SettingsManagerClass.Get()
-				: UMLAdapterManager::StaticClass();
+		UMLAdapterManager& PreviousInstance = UMLAdapterManager::Get();
+		PreviousInstance.ConditionalBeginDestroy();
+	}
+#endif // WITH_EDITOR
 
-			UE_LOG(LogUnrealEditorMLAdapter, Log, TEXT("Creating MLAdapter manager of class %s"), *GetNameSafe(Class));
+	TSubclassOf<UMLAdapterManager> SettingsManagerClass = UMLAdapterSettings::GetManagerClass().Get();
+	UClass* Class = SettingsManagerClass
+		? SettingsManagerClass.Get()
+		: UMLAdapterManager::StaticClass();
 
-			UMLAdapterManager* ManagerInstance = NewObject<UMLAdapterManager>(GEngine, Class);
-			check(ManagerInstance);
-			ManagerInstance->AddToRoot();
+	UE_LOG(LogUnrealEditorMLAdapter, Log, TEXT("Creating MLAdapter manager of class %s"), *GetNameSafe(Class));
 
-			UMLAdapterManager::OnPostInit.Broadcast();
-		}
-	};
+	UMLAdapterManager* NewInstance = NewObject<UMLAdapterManager>(GEngine, Class);
+	check(NewInstance);
+	NewInstance->AddToRoot();
 
-	static FManagerBootloader Loader;
+	UMLAdapterManager::OnPostInit.Broadcast();
 }
 
+struct FManagerBootloader
+{
+	FManagerBootloader()
+	{
+		FCoreDelegates::OnPostEngineInit.AddLambda([this]()
+			{
+				UMLAdapterManager::RecreateManagerInstance();
+			});
+	}
+};
+
+static FManagerBootloader Loader;
+
 //----------------------------------------------------------------------//
 // 
 //
@@ -136,6 +140,7 @@ void UMLAdapterManager::BeginDestroy()
 		ManagerInstance = nullptr;
 	}
 	StopServer();
+	CleanUpDelegates();
 	Super::BeginDestroy();
 }
 
@@ -329,19 +334,33 @@ void UMLAdapterManager::BindToDelegates()
 	//  FWorldDelegates::OnPostWorldCreation;
 	//	FWorldDelegates::OnPreWorldInitialization; 
 
-	FWorldDelegates::OnPostWorldInitialization.AddUObject(this, &UMLAdapterManager::OnPostWorldInit);
-	FWorldDelegates::OnWorldCleanup.AddUObject(this, &UMLAdapterManager::OnWorldCleanup);
+	OnPostWorldInitializationHandle = FWorldDelegates::OnPostWorldInitialization.AddUObject(this, &UMLAdapterManager::OnPostWorldInit);
+	OnWorldCleanupHandle = FWorldDelegates::OnWorldCleanup.AddUObject(this, &UMLAdapterManager::OnWorldCleanup);
 
 	//FGameDelegates
-	FGameModeEvents::OnGameModeInitializedEvent().AddUObject(this, &UMLAdapterManager::OnGameModeInitialized);
+	OnGameModeInitializedHandle = FGameModeEvents::OnGameModeInitializedEvent().AddUObject(this, &UMLAdapterManager::OnGameModeInitialized);
 	// FGameModePreLoginEvent GameModePreLoginEvent;
-	FGameModeEvents::OnGameModePostLoginEvent().AddUObject(this, &UMLAdapterManager::OnGameModePostLogin);
+	OnGameModePostLoginHandle = FGameModeEvents::OnGameModePostLoginEvent().AddUObject(this, &UMLAdapterManager::OnGameModePostLogin);
 	// FGameModeLogoutEvent GameModeLogoutEvent;
-	FGameModeEvents::OnGameModeMatchStateSetEvent().AddUObject(this, &UMLAdapterManager::OnGameModeMatchStateSet);
+	OnGameModeMatchStateSetHandle = FGameModeEvents::OnGameModeMatchStateSetEvent().AddUObject(this, &UMLAdapterManager::OnGameModeMatchStateSet);
+
+#if WITH_EDITORONLY_DATA
+	BeginPIEHandle = FEditorDelegates::BeginPIE.AddUObject(this, &UMLAdapterManager::OnBeginPIE);
+	EndPIEHandle = FEditorDelegates::EndPIE.AddUObject(this, &UMLAdapterManager::OnEndPIE);
+#endif // WITH_EDITORONLY_DATA
+}
+
+void UMLAdapterManager::CleanUpDelegates()
+{
+	FWorldDelegates::OnPostWorldInitialization.Remove(OnPostWorldInitializationHandle);
+	FWorldDelegates::OnWorldCleanup.Remove(OnWorldCleanupHandle);
+	FGameModeEvents::OnGameModeInitializedEvent().Remove(OnGameModeInitializedHandle);
+	FGameModeEvents::OnGameModePostLoginEvent().Remove(OnGameModePostLoginHandle);
+	FGameModeEvents::OnGameModeMatchStateSetEvent().Remove(OnGameModeMatchStateSetHandle);
 
 #if WITH_EDITORONLY_DATA
-	FEditorDelegates::BeginPIE.AddUObject(this, &UMLAdapterManager::OnBeginPIE);
-	FEditorDelegates::EndPIE.AddUObject(this, &UMLAdapterManager::OnEndPIE);
+	FEditorDelegates::BeginPIE.Remove(BeginPIEHandle);
+	FEditorDelegates::EndPIE.Remove(EndPIEHandle);
 #endif // WITH_EDITORONLY_DATA
 }
 
diff --git a/Engine/Plugins/AI/MLAdapter/Source/MLAdapter/Private/MLAdapterSettings.cpp b/Engine/Plugins/AI/MLAdapter/Source/MLAdapter/Private/MLAdapterSettings.cpp
--- a/Engine/Plugins/AI/MLAdapter/Source/MLAdapter/Private/MLAdapterSettings.cpp
+++ b/Engine/Plugins/AI/MLAdapter/Source/MLAdapter/Private/MLAdapterSettings.cpp
@@ -13,7 +13,7 @@ UMLAdapterSettings::UMLAdapterSettings(const FObjectInitializer& ObjectInitializ
 }
 
 TSubclassOf<UMLAdapterManager> UMLAdapterSettings::GetManagerClass()
-{ 
+{
 	const FSoftClassPath LocalClassName = GET_CONFIG_VALUE(ManagerClass);
 	TSubclassOf<UMLAdapterManager> LocalClass = LocalClassName.ResolveClass();
 	return LocalClass;
@@ -37,6 +37,12 @@ TSubclassOf<UMLAdapterAgent> UMLAdapterSettings::GetAgentClass()
 void UMLAdapterSettings::PostEditChangeProperty(struct FPropertyChangedEvent& PropertyChangedEvent)
 {
 	Super::PostEditChangeProperty(PropertyChangedEvent);
+
+	const FName PropertyName = (PropertyChangedEvent.Property != NULL) ? PropertyChangedEvent.Property->GetFName() : NAME_None;
+	if (PropertyName == GET_MEMBER_NAME_CHECKED(UMLAdapterSettings, ManagerClass))
+	{
+		UMLAdapterManager::RecreateManagerInstance();
+	}
 }
 #endif // WITH_EDITOR
 
diff --git a/Engine/Plugins/AI/MLAdapter/Source/MLAdapter/Public/MLAdapterManager.h b/Engine/Plugins/AI/MLAdapter/Source/MLAdapter/Public/MLAdapterManager.h
--- a/Engine/Plugins/AI/MLAdapter/Source/MLAdapter/Public/MLAdapterManager.h
+++ b/Engine/Plugins/AI/MLAdapter/Source/MLAdapter/Public/MLAdapterManager.h
@@ -4,6 +4,7 @@
 
 #include "CoreMinimal.h"
 #include "Delegates/DelegateCombinations.h"
+#include "Delegates/IDelegateInstance.h"
 #include "Logging/LogMacros.h"
 #include "Tickable.h"
 #include "Engine/World.h"
@@ -59,6 +60,7 @@ public:
 	// FExec end
 
 	virtual void BindToDelegates();
+	virtual void CleanUpDelegates();
 
 	virtual void OnPostWorldInit(UWorld* World, const UWorld::InitializationValues);
 	virtual void OnWorldCleanup(UWorld* World, bool bSessionEnded, bool bCleanupResources);
@@ -155,12 +157,23 @@ protected:
 	 *	function) the simulation will progress by StepsRequested ticks before pausing */
 	int32 StepsRequested = 0;
 
+	FDelegateHandle OnPostWorldInitializationHandle;
+	FDelegateHandle OnWorldCleanupHandle;
+	FDelegateHandle OnGameModeInitializedHandle;
+	FDelegateHandle OnGameModePostLoginHandle;
+	FDelegateHandle OnGameModeMatchStateSetHandle;
+#if WITH_EDITORONLY_DATA
+	FDelegateHandle BeginPIEHandle;
+	FDelegateHandle EndPIEHandle;
+#endif // WITH_EDITORONLY_DATA
+
 	static UMLAdapterManager* ManagerInstance;
 public:
+	static void RecreateManagerInstance();
+
 	static FOnGenericEvent OnPostInit;
 };
 
-
 //----------------------------------------------------------------------//
 // inlines 
 //----------------------------------------------------------------------//
