MLAdapter: Moving Managers and Sessions into their own folders

#rnx

diff --git a/Engine/Plugins/AI/MLAdapter/Source/MLAdapter/Private/Actuators/MLAdapterActuator.cpp b/Engine/Plugins/AI/MLAdapter/Source/MLAdapter/Private/Actuators/MLAdapterActuator.cpp
--- a/Engine/Plugins/AI/MLAdapter/Source/MLAdapter/Private/Actuators/MLAdapterActuator.cpp
+++ b/Engine/Plugins/AI/MLAdapter/Source/MLAdapter/Private/Actuators/MLAdapterActuator.cpp
@@ -2,7 +2,7 @@
 
 #include "Actuators/MLAdapterActuator.h"
 #include "MLAdapterTypes.h"
-#include "MLAdapterManager.h"
+#include "Managers/MLAdapterManager.h"
 
 
 namespace
diff --git a/Engine/Plugins/AI/MLAdapter/Source/MLAdapter/Private/Agents/MLAdapterAgents.cpp b/Engine/Plugins/AI/MLAdapter/Source/MLAdapter/Private/Agents/MLAdapterAgents.cpp
--- a/Engine/Plugins/AI/MLAdapter/Source/MLAdapter/Private/Agents/MLAdapterAgents.cpp
+++ b/Engine/Plugins/AI/MLAdapter/Source/MLAdapter/Private/Agents/MLAdapterAgents.cpp
@@ -2,7 +2,7 @@
 
 #include "Agents/MLAdapterAgent.h"
 #include "Sensors/MLAdapterSensor.h"
-#include "MLAdapterSession.h"
+#include "Sessions/MLAdapterSession.h"
 #include "MLAdapterSpace.h"
 #include "MLAdapterLibrarian.h"
 #include "GenericPlatform/GenericPlatformMath.h"
diff --git a/Engine/Plugins/AI/MLAdapter/Source/MLAdapter/Private/Debug/GameplayDebuggerCategory_MLAdapter.cpp b/Engine/Plugins/AI/MLAdapter/Source/MLAdapter/Private/Debug/GameplayDebuggerCategory_MLAdapter.cpp
--- a/Engine/Plugins/AI/MLAdapter/Source/MLAdapter/Private/Debug/GameplayDebuggerCategory_MLAdapter.cpp
+++ b/Engine/Plugins/AI/MLAdapter/Source/MLAdapter/Private/Debug/GameplayDebuggerCategory_MLAdapter.cpp
@@ -3,8 +3,8 @@
 #include "Debug/GameplayDebuggerCategory_MLAdapter.h"
 #include "Engine/Engine.h"
 #include "Misc/CoreDelegates.h"
-#include "MLAdapterManager.h"
-#include "MLAdapterSession.h"
+#include "Managers/MLAdapterManager.h"
+#include "Sessions/MLAdapterSession.h"
 #include "Agents/MLAdapterAgent.h"
 #include "Sensors/MLAdapterSensor.h"
 #include "Actuators/MLAdapterActuator.h"
diff --git a/Engine/Plugins/AI/MLAdapter/Source/MLAdapter/Private/MLAdapterLibrarian.cpp b/Engine/Plugins/AI/MLAdapter/Source/MLAdapter/Private/MLAdapterLibrarian.cpp
--- a/Engine/Plugins/AI/MLAdapter/Source/MLAdapter/Private/MLAdapterLibrarian.cpp
+++ b/Engine/Plugins/AI/MLAdapter/Source/MLAdapter/Private/MLAdapterLibrarian.cpp
@@ -6,7 +6,7 @@
 #include "Actuators/MLAdapterActuator.h"
 #include "Agents/MLAdapterAgent.h"
 #include "UObject/UObjectHash.h"
-#include "MLAdapterManager.h"
+#include "Managers/MLAdapterManager.h"
 
 
 namespace FLibrarianHelper
diff --git a/Engine/Plugins/AI/MLAdapter/Source/MLAdapter/Private/MLAdapterManager.cpp b/Engine/Plugins/AI/MLAdapter/Source/MLAdapter/Private/MLAdapterManager.cpp
deleted file mode 100644
diff --git a/Engine/Plugins/AI/MLAdapter/Source/MLAdapter/Private/MLAdapterManager_Client.cpp b/Engine/Plugins/AI/MLAdapter/Source/MLAdapter/Private/MLAdapterManager_Client.cpp
deleted file mode 100644
diff --git a/Engine/Plugins/AI/MLAdapter/Source/MLAdapter/Private/MLAdapterManager_Server.cpp b/Engine/Plugins/AI/MLAdapter/Source/MLAdapter/Private/MLAdapterManager_Server.cpp
deleted file mode 100644
diff --git a/Engine/Plugins/AI/MLAdapter/Source/MLAdapter/Private/MLAdapterModule.cpp b/Engine/Plugins/AI/MLAdapter/Source/MLAdapter/Private/MLAdapterModule.cpp
--- a/Engine/Plugins/AI/MLAdapter/Source/MLAdapter/Private/MLAdapterModule.cpp
+++ b/Engine/Plugins/AI/MLAdapter/Source/MLAdapter/Private/MLAdapterModule.cpp
@@ -1,7 +1,7 @@
 // Copyright Epic Games, Inc. All Rights Reserved.
 
 #include "MLAdapterModule.h"
-#include "MLAdapterManager.h"
+#include "Managers/MLAdapterManager.h"
 #if WITH_GAMEPLAY_DEBUGGER
 #include "GameplayDebugger.h"
 #include "Debug/GameplayDebuggerCategory_MLAdapter.h"
diff --git a/Engine/Plugins/AI/MLAdapter/Source/MLAdapter/Private/MLAdapterScribe.cpp b/Engine/Plugins/AI/MLAdapter/Source/MLAdapter/Private/MLAdapterScribe.cpp
--- a/Engine/Plugins/AI/MLAdapter/Source/MLAdapter/Private/MLAdapterScribe.cpp
+++ b/Engine/Plugins/AI/MLAdapter/Source/MLAdapter/Private/MLAdapterScribe.cpp
@@ -1,7 +1,7 @@
 // Copyright Epic Games, Inc. All Rights Reserved.
 
 #include "MLAdapterScribe.h"
-#include "MLAdapterManager.h"
+#include "Managers/MLAdapterManager.h"
 #include "MLAdapterLibrarian.h"
 
 #define FSTRING_TO_STD(a) (std::string(TCHAR_TO_UTF8(*a)))
diff --git a/Engine/Plugins/AI/MLAdapter/Source/MLAdapter/Private/MLAdapterSession.cpp b/Engine/Plugins/AI/MLAdapter/Source/MLAdapter/Private/MLAdapterSession.cpp
deleted file mode 100644
diff --git a/Engine/Plugins/AI/MLAdapter/Source/MLAdapter/Private/Managers/MLAdapterManager.cpp b/Engine/Plugins/AI/MLAdapter/Source/MLAdapter/Private/Managers/MLAdapterManager.cpp
new file mode 100644
--- /dev/null
+++ b/Engine/Plugins/AI/MLAdapter/Source/MLAdapter/Private/Managers/MLAdapterManager.cpp
@@ -0,0 +1,545 @@
+// Copyright Epic Games, Inc. All Rights Reserved.
+
+#include "Managers/MLAdapterManager.h"
+#include "CoreGlobals.h"
+#include "Misc/CoreDelegates.h"
+#include "Engine/Engine.h"
+#include "GameFramework/GameModeBase.h"
+#include "MLAdapterTypes.h"
+#include "MLAdapterAsync.h"
+#include "Sessions/MLAdapterSession.h"
+#include "MLAdapterJson.h"
+#include "Agents/MLAdapterAgent.h"
+#include "Sensors/MLAdapterSensor.h"
+#include "MLAdapterSettings.h"
+#include <string>
+#if WITH_EDITORONLY_DATA
+#include "Editor.h"
+#include "Settings/LevelEditorPlaySettings.h"
+#endif // WITH_EDITORONLY_DATA
+
+#include "RPCWrapper/MsgPack.h"
+#include "RPCWrapper/Server.h"
+
+// engine AI support
+#include "AI/NavigationSystemBase.h"
+
+
+namespace
+{
+	FRPCServer* RPCServerInstance = nullptr;
+	
+	EMLAdapterServerMode GetServerModeForWorld(UWorld& World)
+	{
+		EMLAdapterServerMode Mode = EMLAdapterServerMode::Invalid;
+		switch (World.GetNetMode())
+		{
+		case NM_Standalone:
+		case NM_ListenServer:
+			Mode = EMLAdapterServerMode::Standalone;
+			break;
+		case NM_DedicatedServer:
+			Mode = EMLAdapterServerMode::Server;
+			break;
+		case NM_Client:
+			Mode = EMLAdapterServerMode::Client;
+			break;
+		}
+		ensure(Mode != EMLAdapterServerMode::Invalid);
+		return Mode;
+	}
+}
+
+void UMLAdapterManager::RecreateManagerInstance()
+{
+#if WITH_EDITOR
+	if (UMLAdapterManager::IsReady())
+	{
+		UMLAdapterManager& PreviousInstance = UMLAdapterManager::Get();
+		PreviousInstance.ConditionalBeginDestroy();
+	}
+#endif // WITH_EDITOR
+
+	TSubclassOf<UMLAdapterManager> SettingsManagerClass = UMLAdapterSettings::GetManagerClass().Get();
+	UClass* Class = SettingsManagerClass
+		? SettingsManagerClass.Get()
+		: UMLAdapterManager::StaticClass();
+
+	UE_LOG(LogUnrealEditorMLAdapter, Log, TEXT("Creating MLAdapter manager of class %s"), *GetNameSafe(Class));
+
+	UMLAdapterManager* NewInstance = NewObject<UMLAdapterManager>(GEngine, Class);
+	check(NewInstance);
+	NewInstance->AddToRoot();
+
+	UMLAdapterManager::OnPostInit.Broadcast();
+}
+
+struct FManagerBootloader
+{
+	FManagerBootloader()
+	{
+		FCoreDelegates::OnPostEngineInit.AddLambda([this]()
+			{
+				UMLAdapterManager::RecreateManagerInstance();
+			});
+	}
+};
+
+static FManagerBootloader Loader;
+
+//----------------------------------------------------------------------//
+// 
+//
+// @todo we might want to look for a way to restrict users from creating UMLAdapterManager
+// instances manually.
+//----------------------------------------------------------------------//
+UMLAdapterManager* UMLAdapterManager::ManagerInstance = nullptr;
+UMLAdapterManager::FOnGenericEvent UMLAdapterManager::OnPostInit;
+
+UMLAdapterManager::UMLAdapterManager(const FObjectInitializer& ObjectInitializer)
+	: Super(ObjectInitializer)
+{
+	bCommonFunctionsAdded = false;
+	bTickWorldManually = false;
+
+	if (HasAnyFlags(RF_ClassDefaultObject) == false)
+	{
+		ensure(ManagerInstance == nullptr);
+
+		ManagerInstance = this;
+	}
+}
+
+void UMLAdapterManager::PostInitProperties()
+{
+	Super::PostInitProperties();
+	
+	if (HasAnyFlags(RF_ClassDefaultObject) == false)
+	{
+		Librarian.GatherClasses();
+
+		BindToDelegates();
+
+		// if there's any world present create the RPC server
+		UWorld* World = 
+#if WITH_EDITOR
+			GIsEditor ? GWorld :
+#endif // WITH_EDITOR
+			(GEngine->GetWorldContexts().Num() > 0 ? GEngine->GetWorldContexts()[0].World() : nullptr);
+		
+		OnPostWorldInit(World, UWorld::InitializationValues());
+	}
+}
+
+void UMLAdapterManager::BeginDestroy()
+{
+	SetSession(nullptr);
+	if (ManagerInstance == this)
+	{
+		//SetManualWorldTickEnabled(false);
+		ManagerInstance = nullptr;
+	}
+	StopServer();
+	CleanUpDelegates();
+	Super::BeginDestroy();
+}
+
+void UMLAdapterManager::StopServer()
+{
+	if (RPCServerInstance)
+	{
+		UE_LOG(LogUnrealEditorMLAdapter, Log, TEXT("Stopping RPC server."));
+		RPCServerInstance->stop();
+		delete RPCServerInstance;
+		RPCServerInstance = nullptr;
+	}
+	CurrentFunctionMode = EMLAdapterServerMode::Invalid;
+}
+
+void UMLAdapterManager::StartServer(uint16 Port, EMLAdapterServerMode InMode, uint16 ServerThreads)
+{
+	StopServer();
+	
+	RequestedFunctionMode = InMode;
+	ServerThreads = FMath::Max< uint16>(1, ServerThreads);
+
+	EMLAdapterServerMode NewMode = InMode;
+	if (InMode == EMLAdapterServerMode::Invalid || InMode == EMLAdapterServerMode::AutoDetect)
+	{
+		if (LastActiveWorld)
+		{
+			NewMode = GetServerModeForWorld(*LastActiveWorld);
+		}
+		else if (GIsEditor || (GIsClient && GIsServer))
+		{
+			NewMode = EMLAdapterServerMode::Standalone;
+		}
+		else if (GIsClient)
+		{
+			NewMode = EMLAdapterServerMode::Client;
+		}
+		else
+		{
+			NewMode = EMLAdapterServerMode::Server;
+		}
+	}
+
+	UE_LOG(LogUnrealEditorMLAdapter, Log, TEXT("Starting RPC server on port %d."), Port);
+	RPCServerInstance = new FRPCServer(Port);
+	CurrentPort = Port;
+	check(RPCServerInstance);
+
+	bCommonFunctionsAdded = false;
+
+	CurrentFunctionMode = NewMode;
+	switch (NewMode)
+	{
+	case EMLAdapterServerMode::Client:
+		ConfigureAsClient(*RPCServerInstance);
+		break;
+	case EMLAdapterServerMode::Server:
+		ConfigureAsServer(*RPCServerInstance);
+		break;
+	default:
+		ConfigureAsClient(*RPCServerInstance);
+		ConfigureAsServer(*RPCServerInstance);
+		break;
+	}
+	CurrentServerThreads = ServerThreads;
+	RPCServerInstance->async_run(ServerThreads);
+}
+
+void UMLAdapterManager::ConfigureAsStandalone(FRPCServer& Server)
+{
+	ConfigureAsServer(Server);
+	ConfigureAsClient(Server);
+}
+
+void UMLAdapterManager::EnsureAISystemPresence(UWorld& World)
+{
+	if (World.GetAISystem())
+	{
+		return;
+	}
+
+	UAISystemBase* AISystem = World.CreateAISystem();
+	// it's possible the world is configured to not have AI system. Not sure what to do in such a a case
+	ensure(AISystem);
+}
+
+void UMLAdapterManager::EnsureNavigationSystemPresence(UWorld& World)
+{
+	if (World.GetNavigationSystem())
+	{
+		return;
+	}
+	
+	FNavigationSystem::AddNavigationSystemToWorld(World);	
+	// it's possible the world is configured to not have a nav system. Not sure what to do in such a a case
+	ensure(World.GetNavigationSystem());
+}
+
+bool UMLAdapterManager::IsRunning() const
+{
+	return (RPCServerInstance != nullptr);
+}
+
+TStatId UMLAdapterManager::GetStatId() const
+{
+	RETURN_QUICK_DECLARE_CYCLE_STAT(UMLAdapterManager, STATGROUP_Tickables);
+}
+
+ETickableTickType UMLAdapterManager::GetTickableTickType() const 
+{ 
+	return HasAnyFlags(RF_ClassDefaultObject) ? ETickableTickType::Never : ETickableTickType::Always;
+}
+
+bool UMLAdapterManager::IsTickable() const 
+{ 
+	return (HasAnyFlags(RF_ClassDefaultObject) == false); 
+}
+
+void UMLAdapterManager::Tick(float DeltaTime)
+{
+	if (IsWorldRealTime() || StepsRequested > 0)
+	{
+		if (Session)
+		{
+			DeltaTime = IsWorldRealTime() ? DeltaTime : (1.f / WorldFPS);
+			Session->Tick(DeltaTime);
+		}
+		--StepsRequested;
+	}
+}
+
+UMLAdapterSession* UMLAdapterManager::CreateNewSession()
+{
+	UClass* Class = UMLAdapterSettings::GetSessionClass().Get()
+		? UMLAdapterSettings::GetSessionClass().Get()
+		: UMLAdapterSession::StaticClass();
+
+	UMLAdapterSession* NewSession = FMLAdapter::NewObject<UMLAdapterSession>(this, Class);
+	check(NewSession);
+	NewSession->SetWorld(LastActiveWorld);
+	// some config
+	return NewSession;
+}
+
+void UMLAdapterManager::SetSession(UMLAdapterSession* NewSession)
+{
+	if (Session != nullptr && (NewSession == nullptr || Session != NewSession))
+	{
+		Session->Close();
+		// clear the flag to let GC remove the object
+		Session->ClearInternalFlags(EInternalObjectFlags::Async);
+		Session = nullptr;
+	}
+	Session = NewSession;
+	if (Session)
+	{
+		// we're going to be using this object in a async manner, so we need to 
+		// mark it appropriately. This will make GC ignore this object until we 
+		// clear the flag. 
+		Session->SetInternalFlags(EInternalObjectFlags::Async);
+		Session->Open();
+	}
+
+	OnCurrentSessionChanged.Broadcast();
+}
+
+void UMLAdapterManager::CloseSession(UMLAdapterSession& InSession)
+{
+	// @todo temporary implementation, will change with multi-session support
+	if (&InSession == Session)
+	{
+		SetSession(nullptr);
+	}
+}
+
+UMLAdapterSession& UMLAdapterManager::GetSession()
+{
+	if (Session == nullptr)
+	{
+		SetSession(CreateNewSession());
+	}
+	check(Session);
+	return *Session;
+}
+
+void UMLAdapterManager::BindToDelegates()
+{	
+	// Commented out possible other useful delegates
+	//	FCoreDelegates::GameStateClassChanged;
+	//	FCoreDelegates::ConfigReadyForUse;
+	//  FWorldDelegates::OnPostWorldCreation;
+	//	FWorldDelegates::OnPreWorldInitialization; 
+
+	OnPostWorldInitializationHandle = FWorldDelegates::OnPostWorldInitialization.AddUObject(this, &UMLAdapterManager::OnPostWorldInit);
+	OnWorldCleanupHandle = FWorldDelegates::OnWorldCleanup.AddUObject(this, &UMLAdapterManager::OnWorldCleanup);
+
+	//FGameDelegates
+	OnGameModeInitializedHandle = FGameModeEvents::OnGameModeInitializedEvent().AddUObject(this, &UMLAdapterManager::OnGameModeInitialized);
+	// FGameModePreLoginEvent GameModePreLoginEvent;
+	OnGameModePostLoginHandle = FGameModeEvents::OnGameModePostLoginEvent().AddUObject(this, &UMLAdapterManager::OnGameModePostLogin);
+	// FGameModeLogoutEvent GameModeLogoutEvent;
+	OnGameModeMatchStateSetHandle = FGameModeEvents::OnGameModeMatchStateSetEvent().AddUObject(this, &UMLAdapterManager::OnGameModeMatchStateSet);
+
+#if WITH_EDITORONLY_DATA
+	BeginPIEHandle = FEditorDelegates::BeginPIE.AddUObject(this, &UMLAdapterManager::OnBeginPIE);
+	EndPIEHandle = FEditorDelegates::EndPIE.AddUObject(this, &UMLAdapterManager::OnEndPIE);
+#endif // WITH_EDITORONLY_DATA
+}
+
+void UMLAdapterManager::CleanUpDelegates()
+{
+	FWorldDelegates::OnPostWorldInitialization.Remove(OnPostWorldInitializationHandle);
+	FWorldDelegates::OnWorldCleanup.Remove(OnWorldCleanupHandle);
+	FGameModeEvents::OnGameModeInitializedEvent().Remove(OnGameModeInitializedHandle);
+	FGameModeEvents::OnGameModePostLoginEvent().Remove(OnGameModePostLoginHandle);
+	FGameModeEvents::OnGameModeMatchStateSetEvent().Remove(OnGameModeMatchStateSetHandle);
+
+#if WITH_EDITORONLY_DATA
+	FEditorDelegates::BeginPIE.Remove(BeginPIEHandle);
+	FEditorDelegates::EndPIE.Remove(EndPIEHandle);
+#endif // WITH_EDITORONLY_DATA
+}
+
+bool UMLAdapterManager::ShouldInitForWorld(const UWorld& World) const
+{
+	if (!World.IsGameWorld())
+	{
+		return false;
+	}
+
+#if WITH_EDITOR
+	// PIE is a special case, we need to see if it's a client-server PIE and 
+	// if so we need to filter the incoming world based on the settings
+	if (World.WorldType == EWorldType::PIE)
+	{
+		if (World.HasAnyFlags(RF_WasLoaded))
+		{
+			const ULevelEditorPlaySettings* PlayInSettings = GetDefault<ULevelEditorPlaySettings>();
+			if (PlayInSettings)
+			{
+				EPlayNetMode PlayNetMode = PIE_Standalone;
+				PlayInSettings->GetPlayNetMode(PlayNetMode);
+
+				if (PlayNetMode != PIE_Standalone)
+				{
+					UE_LOG(LogUnrealEditorMLAdapter, Log, TEXT("Ignoring %s due to net mode != PIE_Standalone"), *World.GetName());
+					return false;
+				}
+			}
+			else
+			{
+				// skipping temp world
+				return false;
+			}
+		}
+
+	}
+#endif // WITH_EDITOR
+
+	return true;
+}
+
+void UMLAdapterManager::OnPostWorldInit(UWorld* World, const UWorld::InitializationValues)
+{
+	if (World && ShouldInitForWorld(*World))
+	{
+		LastActiveWorld = World;
+		if ((RequestedFunctionMode == EMLAdapterServerMode::Invalid || RequestedFunctionMode == EMLAdapterServerMode::AutoDetect)
+			&& GetServerModeForWorld(*World) != CurrentFunctionMode)
+		{
+			// restart the RPC server. Note that this will kick all the currently connected agents
+			uint16 Port = UMLAdapterSettings::GetDefaultRPCServerPort();
+			FParse::Value(FCommandLine::Get(), TEXT("MLAdapterPort="), Port);
+			StartServer(Port, GetServerModeForWorld(*World), CurrentServerThreads);
+		}
+
+		if (HasSession())
+		{
+			GetSession().OnPostWorldInit(*World);
+		}
+	}
+}
+
+void UMLAdapterManager::OnWorldCleanup(UWorld* World, bool bSessionEnded, bool bCleanupResources)
+{
+	// no need to remove, the World is going away
+	if (World && World->IsGameWorld())
+	{
+		if (World == LastActiveWorld)
+		{
+			LastActiveWorld = nullptr;
+		}
+
+		if (HasSession())
+		{
+			CloseSession(GetSession());
+		}
+	}
+}
+
+void UMLAdapterManager::OnBeginPIE(const bool bIsSimulating)
+{
+
+}
+
+void UMLAdapterManager::OnEndPIE(const bool bIsSimulating)
+{
+
+}
+
+void UMLAdapterManager::OnGameModeInitialized(AGameModeBase* GameMode)
+{
+	if (GameMode && HasSession())
+	{
+		GetSession().OnGameModeInitialized(*GameMode);
+	}
+}
+
+void UMLAdapterManager::OnGameModePostLogin(AGameModeBase* GameMode, APlayerController* NewPlayer)
+{
+	if (GameMode && NewPlayer && HasSession())
+	{
+		GetSession().OnGameModePostLogin(GameMode, NewPlayer);
+	}
+}
+
+void UMLAdapterManager::OnGameModeMatchStateSet(FName MatchState)
+{
+	if (HasSession())
+	{
+		GetSession().OnGameModeMatchStateSet(MatchState);
+	}
+}
+
+void UMLAdapterManager::ResetWorld()
+{
+	if (LastActiveWorld)
+	{
+		AGameModeBase* GameMode = LastActiveWorld->GetAuthGameMode<AGameModeBase>();
+		if (GameMode)
+		{
+			GameMode->ResetLevel();
+		}
+	}
+
+	if (HasSession())
+	{
+		GetSession().ResetWorld();
+	}
+}
+
+void UMLAdapterManager::SetManualWorldTickEnabled(bool bEnable)
+{
+	bTickWorldManually = bEnable;
+	if (Session)
+	{
+		Session->SetManualWorldTickEnabled(bEnable);
+	}
+}
+
+bool UMLAdapterManager::Exec(class UWorld* InWorld, const TCHAR* Cmd, FOutputDevice& Ar)
+{
+	if (HasAnyFlags(RF_ClassDefaultObject))
+	{
+		return false;
+	}
+
+	if (FParse::Command(&Cmd, TEXT("mladapter_session_stop")))
+	{
+		SetSession(nullptr);
+		return true;
+	}
+
+	return false;
+}
+
+namespace FMLAdapterConsoleCommands
+{
+	struct FHelper
+	{
+		static void RestartServer(const TArray<FString>& Args, UWorld*)
+		{
+			UMLAdapterManager& Manager = UMLAdapterManager::Get();
+			uint16 Port = UMLAdapterSettings::GetDefaultRPCServerPort();
+			FParse::Value(FCommandLine::Get(), TEXT("MLAdapterPort="), Port);
+			if (Args.Num() > 0)
+			{
+				Port = uint16(TCString<TCHAR>::Atoi(*Args[0]));
+			}
+			Manager.StartServer(Port, Manager.CurrentFunctionMode, Manager.CurrentServerThreads);
+		}
+	};
+
+	FAutoConsoleCommand StopSession(TEXT("mladapter.session.stop"), TEXT(""), FConsoleCommandDelegate::CreateLambda([]()
+	{
+		UMLAdapterManager::Get().SetSession(nullptr);
+	}));
+
+	FAutoConsoleCommandWithWorldAndArgs RestartServer(TEXT("mladapter.server.restart")
+		, TEXT("restarts the MLAdapter RPC server, optionally changing the port the server is listening at. Use: mladapter.server.restart [port]")
+		, FConsoleCommandWithWorldAndArgsDelegate::CreateStatic(FHelper::RestartServer));
+}
diff --git a/Engine/Plugins/AI/MLAdapter/Source/MLAdapter/Private/Managers/MLAdapterManager_Client.cpp b/Engine/Plugins/AI/MLAdapter/Source/MLAdapter/Private/Managers/MLAdapterManager_Client.cpp
new file mode 100644
--- /dev/null
+++ b/Engine/Plugins/AI/MLAdapter/Source/MLAdapter/Private/Managers/MLAdapterManager_Client.cpp
@@ -0,0 +1,312 @@
+// Copyright Epic Games, Inc. All Rights Reserved.
+
+#include "Managers/MLAdapterManager.h"
+#include "MLAdapterTypes.h"
+#include "MLAdapterAsync.h"
+#include "Sessions/MLAdapterSession.h"
+#include "MLAdapterScribe.h"
+#include "MLAdapterSpace.h"
+#include "Agents/MLAdapterAgent.h"
+#include "Actuators/MLAdapterActuator.h"
+#include "Engine/GameEngine.h"
+
+#include "RPCWrapper/Server.h"
+
+#define REPORT_NOT_IMPLEMENTED() rpc::this_handler().respond_error("Not implemented yet")
+#define CHECK_AGENTID(AgentID) { \
+	if (HasSession() == false) \
+	{ \
+		rpc::this_handler().respond_error("No active session"); \
+	}\
+	if (GetSession().GetAgent(AgentID) == nullptr) \
+	{ \
+		rpc::this_handler().respond_error(std::make_tuple("No Agent of ID", AgentID)); \
+	} \
+} \
+
+
+void UMLAdapterManager::AddCommonFunctions(FRPCServer& Server)
+{
+	if (bCommonFunctionsAdded)
+	{
+		return;
+	}
+
+#if WITH_RPCLIB
+	Server.bind("list_functions", &FMLAdapterScribe::ListFunctions);
+	Librarian.AddRPCFunctionDescription(TEXT("list_functions"), TEXT("(), Lists all functions available through RPC"));
+
+	Server.bind("get_description", [](std::string const& ElementName) {
+		return FMLAdapterScribe::GetDescription(ElementName);
+	});
+	Librarian.AddRPCFunctionDescription(TEXT("get_description"), TEXT("(string ElementName), Describes given element"));
+
+	Server.bind("list_sensor_types", &FMLAdapterScribe::ListSensorTypes);
+	Librarian.AddRPCFunctionDescription(TEXT("list_sensor_types"), TEXT("(), Lists all sensor types available to agents. Note that some of sensors might not make sense in a given environment (like reading keyboard in an mouse-only game)."));
+
+	Server.bind("list_actuator_types", &FMLAdapterScribe::ListActuatorTypes);
+	Librarian.AddRPCFunctionDescription(TEXT("list_actuator_types"), TEXT("(), Lists all actuator types available to agents. Note that some of actuators might not make sense in a given environment (like faking keyboard actions in an mouse-only game)."));
+
+	Server.bind("ping", []() { return true; });
+	Librarian.AddRPCFunctionDescription(TEXT("ping"), TEXT("(), Checks if the RPC server is still alive and responding."));
+
+	Server.bind("get_name", []() {
+		return std::string(TCHAR_TO_UTF8(GInternalProjectName));
+	});
+	Librarian.AddRPCFunctionDescription(TEXT("get_name"), TEXT("(), Fetches a human-readable identifier of the environment the external client is connected to."));
+
+	Server.bind("is_finished", [this](FMLAdapter::FAgentID AgentID) {
+		return HasSession() == false || GetSession().IsDone() || GetSession().GetAgent(AgentID) == nullptr 
+			|| GetSession().GetAgent(AgentID)->IsDone();
+	});
+	Librarian.AddRPCFunctionDescription(TEXT("is_finished"), TEXT("(agent_id), Checks if the game/simulation/episode is done for given agent_id."));
+
+	Server.bind("exit", []() {
+		FPlatformMisc::RequestExit(/*bForce=*/false);
+	});
+	Librarian.AddRPCFunctionDescription(TEXT("exit"), TEXT("(), Closes the UnrealEngine instance."));
+	
+	Server.bind("batch_is_finished", [this](std::vector<FMLAdapter::FAgentID> AgentIDs) {
+		std::vector<bool> Results;
+		if (HasSession() == false || GetSession().IsDone())
+		{
+			for (int Index = 0; Index < int(AgentIDs.size()); ++Index)
+			{
+				Results.push_back(true);
+			}
+		}
+		else
+		{
+			UMLAdapterSession& SessionInstance = GetSession();
+			for (const FMLAdapter::FAgentID& AgentID : AgentIDs)
+			{
+				UMLAdapterAgent* Agent = SessionInstance.GetAgent(AgentID);
+				Results.push_back(Agent == nullptr || Agent->IsDone());
+			}
+		}
+		return Results;
+	});
+	Librarian.AddRPCFunctionDescription(TEXT("batch_is_finished"), TEXT("(), Multi-agent version of is_finished"));
+#endif // WITH_RPCLIB
+
+	bCommonFunctionsAdded = true;
+}
+
+void UMLAdapterManager::ConfigureAsClient(FRPCServer& Server)
+{
+	UE_LOG(LogUnrealEditorMLAdapter, Log, TEXT("\tconfiguring as client"));
+
+	AddCommonFunctions(Server);
+
+#if WITH_RPCLIB
+	Server.bind("add_agent", [this]() {
+		return CallOnGameThread<FMLAdapter::FAgentID>([this]()
+		{
+			return GetSession().AddAgent();
+		});
+	});
+	Librarian.AddRPCFunctionDescription(TEXT("add_agent"), TEXT("Adds a default agent for current environment. Returns added agent's ID if successful, uint(-1) if failed."));
+
+	Server.bind("get_agent_config", [this](FMLAdapter::FAgentID AgentID) {
+		CHECK_AGENTID(AgentID);
+		const FMLAdapterAgentConfig& Config = GetSession().GetAgent(AgentID)->GetConfig();
+		return FSTRING_TO_STD(FMLAdapter::StructToJsonString(Config));
+	});
+	Librarian.AddRPCFunctionDescription(TEXT("get_agent_config"), TEXT("(uint AgentID), Retrieves given agent's config in JSON formatted string"));
+
+	Server.bind("act", [this](FMLAdapter::FAgentID AgentID, std::vector<float> ValueStream) {
+		CHECK_AGENTID(AgentID);
+		UMLAdapterAgent* Agent = GetSession().GetAgent(AgentID);
+		check(Agent);
+
+		const uint8* DataPtr = (const uint8*)ValueStream.data();
+		TArray<uint8> Buffer;
+		Buffer.Append(DataPtr, ValueStream.size() * sizeof(float));
+		FMLAdapterMemoryReader Reader(Buffer);
+		Agent->DigestActions(Reader);
+	});
+	Librarian.AddRPCFunctionDescription(TEXT("act"), TEXT("(uint agent_id, list actions), Distributes the given values array amongst all the actuators, based on actions_space."));
+
+	Server.bind("batch_act", [this](std::vector<FMLAdapter::FAgentID> AgentIDs, std::vector<std::vector<float>> ValueStreams) {
+		if (HasSession() == false)
+		{
+			rpc::this_handler().respond_error("No active session");
+		}
+
+		for (int Index = 0; Index < int(AgentIDs.size()); ++Index)
+		{
+			UMLAdapterAgent* Agent = GetSession().GetAgent(AgentIDs[Index]);
+			if (Agent)
+			{
+				const uint8* DataPtr = (const uint8*)ValueStreams[Index].data();
+				TArray<uint8> Buffer;
+				Buffer.Append(DataPtr, ValueStreams[Index].size() * sizeof(float));
+				FMLAdapterMemoryReader Reader(Buffer);
+				Agent->DigestActions(Reader);
+			}
+		}
+	});
+	Librarian.AddRPCFunctionDescription(TEXT(""), TEXT("A multi-agent version of \'act\' function"));
+
+	Server.bind("get_observations", [this](FMLAdapter::FAgentID AgentID) {
+		std::vector<float> Values;
+		if (HasSession() && GetSession().GetAgent(AgentID))
+		{
+			TArray<uint8> Buffer;
+			FMLAdapterMemoryWriter Writer(Buffer);
+			//Observations.TimeStamp = GetSession().GetTimestamp();
+			GetSession().GetAgent(AgentID)->GetObservations(Writer);
+
+			const float* DataPtr = (float*)Buffer.GetData();
+			Values.assign(DataPtr, DataPtr + Buffer.Num() / sizeof(float));
+		}
+		return Values;
+	});
+	Librarian.AddRPCFunctionDescription(TEXT("get_observations"), TEXT("(uint agent_id), fetches all the information gathered by given agent's sensors. Result matches observations_space"));
+
+	Server.bind("batch_get_observations", [this](std::vector<FMLAdapter::FAgentID> AgentIDs) {
+		std::vector<std::vector<float>> Values;
+		if (HasSession())
+		{
+			Values.resize(AgentIDs.size());
+			for (int Index = 0; Index < int(AgentIDs.size()); ++Index)
+			{
+				UMLAdapterAgent* Agent = GetSession().GetAgent(AgentIDs[Index]);
+				if (Agent)
+				{
+					TArray<uint8> Buffer;
+					FMLAdapterMemoryWriter Writer(Buffer);
+					Agent->GetObservations(Writer);
+
+					const float* DataPtr = (float*)Buffer.GetData();
+					Values[Index].assign(DataPtr, DataPtr + Buffer.Num() / sizeof(float));
+				}
+			}
+		}
+		return Values;
+	});
+	Librarian.AddRPCFunctionDescription(TEXT("batch_get_observations"), TEXT("Multi-agent version of 'get_observations'"));
+
+	Server.bind("get_recent_agent", [this]() {
+		return HasSession() ? FMLAdapter::FAgentID(GetSession().GetAgentsCount() - 1) : FMLAdapter::InvalidAgentID;
+	});
+	Librarian.AddRPCFunctionDescription(TEXT("get_recent_agent"), TEXT("(), Fetches ID of the most recently created agent."));
+
+	Server.bind("get_reward", [this](FMLAdapter::FAgentID AgentID) {
+		CHECK_AGENTID(AgentID);
+		return GetSession().GetAgent(AgentID)->GetReward();
+	});
+	Librarian.AddRPCFunctionDescription(TEXT("get_reward"), TEXT("(uint agent_id), Fetch current reward for given Agent."));
+
+	Server.bind("batch_get_rewards", [this](std::vector<FMLAdapter::FAgentID> AgentIDs) {
+		if (HasSession() == false)
+		{
+			rpc::this_handler().respond_error("No active session");
+		}
+		std::vector<float> Rewards;
+		for (const FMLAdapter::FAgentID& AgentID : AgentIDs)
+		{
+			UMLAdapterAgent* Agent = GetSession().GetAgent(AgentID);
+			Rewards.push_back(Agent ? Agent->GetReward() : 0.f);
+		}
+
+		return Rewards;
+	});
+	Librarian.AddRPCFunctionDescription(TEXT("batch_get_rewards"), TEXT("(), Multi-agent version of 'get_rewards'."));
+
+	Server.bind("desc_action_space", [this](FMLAdapter::FAgentID AgentID) {
+		return CallOnGameThread<std::string>([this, AgentID]()
+		{
+			CHECK_AGENTID(AgentID);
+			FMLAdapterSpaceDescription SpaceDesc;
+			GetSession().GetAgent(AgentID)->GetActionSpaceDescription(SpaceDesc);
+			return FSTRING_TO_STD(SpaceDesc.ToJson());
+		});
+	});
+	Librarian.AddRPCFunctionDescription(TEXT("desc_action_space"), TEXT("(uint agent_id), Fetches actions space desction for given agent"));
+
+	// we're sending this call to game thread since if it's called right after 
+	// "configure_agent" then this call will fetch pre-config state due to agent 
+	// configuration being performed on the game thread
+	Server.bind("desc_observation_space", [this](FMLAdapter::FAgentID AgentID) {
+		return CallOnGameThread<std::string>([this, AgentID]()
+		{
+			CHECK_AGENTID(AgentID);
+			FMLAdapterSpaceDescription SpaceDesc;
+			GetSession().GetAgent(AgentID)->GetObservationSpaceDescription(SpaceDesc);
+			return FSTRING_TO_STD(SpaceDesc.ToJson());
+		});
+	}); 
+	Librarian.AddRPCFunctionDescription(TEXT("desc_observation_space"), TEXT("(uint agent_id), Fetches observations space desction for given agent"));
+		
+	Server.bind("reset", []() {
+		CallOnGameThread<void>([]()
+		{
+			UMLAdapterManager::Get().ResetWorld();
+		});
+	});
+	Librarian.AddRPCFunctionDescription(TEXT("reset"), TEXT("(), Lets the MLAdapter manager know that the environments should be reset. The details of how this call is handles heavily depends on the environment itself."));
+
+	Server.bind("disconnect", [this](FMLAdapter::FAgentID AgentID) {
+		CHECK_AGENTID(AgentID);
+		GetSession().RemoveAgent(AgentID);
+	});
+	Librarian.AddRPCFunctionDescription(TEXT(""), TEXT("(uint agent_id), Lets the MLAdapter session know that given agent will not continue and is to be removed from the session."));
+
+	// calling this means we're done messing up with the agent (configuring and all) 
+	// and we're ready to roll
+	Server.bind("configure_agent", [this](FMLAdapter::FAgentID AgentID, std::string const& JsonConfigString) {
+
+		if (HasSession() == false)
+		{
+			rpc::this_handler().respond_error("No active session");
+			return;
+		}
+		else if (GetSession().GetAgent(AgentID) == nullptr)
+		{
+			rpc::this_handler().respond_error(std::make_tuple("No Agent of ID", AgentID));
+			return;
+		}
+
+		FMLAdapterAgentConfig Config;
+		FMLAdapter::JsonStringToStruct(FString(JsonConfigString.c_str()), Config);
+
+		CallOnGameThread<void>([this, Config, AgentID]()
+		{
+			GetSession().GetAgent(AgentID)->Configure(Config);
+		});
+	});
+	Librarian.AddRPCFunctionDescription(TEXT("configure_agent"), TEXT("(uint agent_id, string json_config), Configures given agent based on json_config. Will throw an exception if given agent doesn't exist."));
+
+	// combines 'add' and 'configure' agent
+	Server.bind("create_agent", [this](std::string const& JsonConfigString) {
+		FMLAdapterAgentConfig Config;
+		FMLAdapter::JsonStringToStruct(FString(JsonConfigString.c_str()), Config);
+
+		return CallOnGameThread<FMLAdapter::FAgentID>([this, Config]()
+		{
+			return GetSession().AddAgent(Config);
+		});
+	});
+	Librarian.AddRPCFunctionDescription(TEXT("create_agent"), TEXT("(), Creates a new agent and returns its agent_id."));
+
+	Server.bind("is_agent_ready", [this](FMLAdapter::FAgentID AgentID) {
+		return HasSession() && GetSession().IsAgentReady(AgentID);
+	});
+	Librarian.AddRPCFunctionDescription(TEXT("is_agent_ready"), TEXT("(uint agent_id), Returns 'true' if given agent is ready to play, including having an avatar"));
+
+	Server.bind("is_ready", [this]() {
+		return CallOnGameThread<bool>([this]()
+		{
+			return HasSession() && GetSession().IsReady();
+		});
+	});
+	Librarian.AddRPCFunctionDescription(TEXT("is_ready"), TEXT("(), return whether the session is ready to go, i.e. whether the simulation has loaded and started."));
+#endif // WITH_RPCLIB
+
+	if (Session)
+	{
+		Session->ConfigureAsClient();
+	}
+	OnAddClientFunctions.Broadcast(Server);
+}
\ No newline at end of file
diff --git a/Engine/Plugins/AI/MLAdapter/Source/MLAdapter/Private/Managers/MLAdapterManager_Server.cpp b/Engine/Plugins/AI/MLAdapter/Source/MLAdapter/Private/Managers/MLAdapterManager_Server.cpp
new file mode 100644
--- /dev/null
+++ b/Engine/Plugins/AI/MLAdapter/Source/MLAdapter/Private/Managers/MLAdapterManager_Server.cpp
@@ -0,0 +1,58 @@
+// Copyright Epic Games, Inc. All Rights Reserved.
+
+#include "Managers/MLAdapterManager.h"
+#include "CoreGlobals.h"
+#include "Misc/CoreDelegates.h"
+#include "Engine/Engine.h"
+#include "GameFramework/GameModeBase.h"
+#include "MLAdapterTypes.h"
+#include "MLAdapterAsync.h"
+#include "Sessions/MLAdapterSession.h"
+#include "MLAdapterJson.h"
+#include "Agents/MLAdapterAgent.h"
+#include "Sensors/MLAdapterSensor.h"
+#include "MLAdapterSettings.h"
+#include <string>
+#if WITH_EDITORONLY_DATA
+#include "Editor.h"
+#endif // WITH_EDITORONLY_DATA
+
+#include "RPCWrapper/Server.h"
+
+
+void UMLAdapterManager::ConfigureAsServer(FRPCServer& Server)
+{
+	UE_LOG(LogUnrealEditorMLAdapter, Log, TEXT("\tconfiguring as server"));
+
+	AddCommonFunctions(Server);
+
+#if WITH_RPCLIB
+	Server.bind("enable_manual_world_tick", [this](bool bEnable) {
+		SetManualWorldTickEnabled(bEnable);
+	});
+	Librarian.AddRPCFunctionDescription(TEXT("enable_manual_world_tick"), TEXT("(), Controlls whether the world is running real time or it\'s being ticked manually with calls to \'step\' or \'request_world_tick\' functions. Default is \'real time\'."));
+
+	Server.bind("request_world_tick", [this](int32 TickCount, bool bWaitForWorldTick) {
+		if (bTickWorldManually == false)
+		{
+			return;
+		}
+		StepsRequested = TickCount;
+		while (bWaitForWorldTick && StepsRequested > 0)
+		{
+			FPlatformProcess::Sleep(0.f);
+		}
+		return;
+	});
+	Librarian.AddRPCFunctionDescription(TEXT("request_world_tick"), TEXT("(int TickCount, bool bWaitForWorldTick), Requests a TickCount world ticks. This has meaning only if \'enable_manual_world_tick(true)\' has been called prior to this function. If bWaitForWorldTick is true then the call will not return until the world has been ticked required number of times"));
+
+	Server.bind("close_session", [this]() { UMLAdapterManager::Get().SetSession(nullptr); });
+	Librarian.AddRPCFunctionDescription(TEXT("close_session"), TEXT("(), shuts down the current session (along with all the agents)."));
+#endif // WITH_RPCLIB
+
+	if (Session)
+	{
+		Session->ConfigureAsServer();
+	}
+	OnAddServerFunctions.Broadcast(Server);
+}
\ No newline at end of file
diff --git a/Engine/Plugins/AI/MLAdapter/Source/MLAdapter/Private/Sensors/MLAdapterSensor.cpp b/Engine/Plugins/AI/MLAdapter/Source/MLAdapter/Private/Sensors/MLAdapterSensor.cpp
--- a/Engine/Plugins/AI/MLAdapter/Source/MLAdapter/Private/Sensors/MLAdapterSensor.cpp
+++ b/Engine/Plugins/AI/MLAdapter/Source/MLAdapter/Private/Sensors/MLAdapterSensor.cpp
@@ -3,7 +3,7 @@
 #include "Sensors/MLAdapterSensor.h"
 #include "Agents/MLAdapterAgent.h"
 #include "MLAdapterTypes.h"
-#include "MLAdapterManager.h"
+#include "Managers/MLAdapterManager.h"
 
 
 namespace
diff --git a/Engine/Plugins/AI/MLAdapter/Source/MLAdapter/Private/Sensors/MLAdapterSensor_AIPerception.cpp b/Engine/Plugins/AI/MLAdapter/Source/MLAdapter/Private/Sensors/MLAdapterSensor_AIPerception.cpp
--- a/Engine/Plugins/AI/MLAdapter/Source/MLAdapter/Private/Sensors/MLAdapterSensor_AIPerception.cpp
+++ b/Engine/Plugins/AI/MLAdapter/Source/MLAdapter/Private/Sensors/MLAdapterSensor_AIPerception.cpp
@@ -5,7 +5,7 @@
 #include "GameFramework/Controller.h"
 #include "Perception/AIPerceptionComponent.h"
 #include "Debug/DebugHelpers.h"
-#include "MLAdapterManager.h"
+#include "Managers/MLAdapterManager.h"
 #include "Perception/AISenseConfig_Sight.h"
 #include "MLAdapterSpace.h"
 
diff --git a/Engine/Plugins/AI/MLAdapter/Source/MLAdapter/Private/Sessions/MLAdapterSession.cpp b/Engine/Plugins/AI/MLAdapter/Source/MLAdapter/Private/Sessions/MLAdapterSession.cpp
new file mode 100644
--- /dev/null
+++ b/Engine/Plugins/AI/MLAdapter/Source/MLAdapter/Private/Sessions/MLAdapterSession.cpp
@@ -0,0 +1,605 @@
+// Copyright Epic Games, Inc. All Rights Reserved.
+
+#include "Sessions/MLAdapterSession.h"
+#include "MLAdapterTypes.h"
+#include "MLAdapterSettings.h"
+#include "EngineUtils.h"
+#include "Engine/GameInstance.h"
+#include "GameFramework/GameMode.h"
+#include "GameFramework/GameStateBase.h"
+#include "Agents/MLAdapterAgent.h"
+#include "Managers/MLAdapterManager.h"
+
+
+void UMLAdapterSession::PostInitProperties()
+{
+	Super::PostInitProperties();
+
+	if (HasAnyFlags(RF_ClassDefaultObject) == false)
+	{
+		bTickWorldManually = (UMLAdapterManager::Get().IsWorldRealTime() == false);
+	}
+}
+
+void UMLAdapterSession::BeginDestroy()
+{
+	Super::BeginDestroy();
+}
+
+void UMLAdapterSession::SetManualWorldTickEnabled(bool bEnable)
+{
+	bTickWorldManually = bEnable;
+	if (bEnable)
+	{
+		if (WorldTicker == nullptr && CachedWorld)
+		{
+			WorldTicker = MakeShareable(new FWorldTicker(CachedWorld));
+		}
+	}
+	else
+	{
+		WorldTicker = nullptr;
+	}
+}
+
+void UMLAdapterSession::FWorldTicker::Tick(float DeltaTime) 
+{
+#if WITH_EDITORONLY_DATA
+	if (CachedWorld.IsValid())
+	{
+		// will get cleared in Session::tick
+		GIntraFrameDebuggingGameThread = true;
+	}
+#endif // WITH_EDITORONLY_DATA
+}
+
+UMLAdapterSession::FWorldTicker::~FWorldTicker()
+{
+#if WITH_EDITORONLY_DATA
+	GIntraFrameDebuggingGameThread = false;
+#endif // WITH_EDITORONLY_DATA
+}
+
+void UMLAdapterSession::SetWorld(UWorld* NewWorld)
+{
+	if (CachedWorld == NewWorld)
+	{
+		return;
+	}
+
+	WorldTicker = nullptr;
+
+	if (CachedWorld)
+	{
+		RemoveAvatars(CachedWorld);
+
+		SetGameMode(nullptr);
+		CachedWorld->RemoveOnActorSpawnedHandler(ActorSpawnedDelegateHandle);
+		ActorSpawnedDelegateHandle.Reset();
+		CachedWorld = nullptr;
+		
+		LastTimestamp = -1.f;
+	}
+
+	if (NewWorld)
+	{
+		CachedWorld = NewWorld;
+		LastTimestamp = NewWorld->GetTimeSeconds();
+		SetGameMode(NewWorld->GetAuthGameMode());
+
+		if (bTickWorldManually)
+		{
+			WorldTicker = MakeShareable(new FWorldTicker(NewWorld));
+		}
+
+		FindAvatars(*CachedWorld);
+		ActorSpawnedDelegateHandle = CachedWorld->AddOnActorSpawnedHandler(FOnActorSpawned::FDelegate::CreateUObject(this, &UMLAdapterSession::OnActorSpawned));
+	}
+}
+
+void UMLAdapterSession::OnActorSpawned(AActor* InActor)
+{
+	// check if it's something we need to consider, like a pawn or controller
+
+	// CONCERN: this function is going to be called a lot and we care only about
+	// a couple of possible actor classes. Sometimes we don't event care at all. 
+	// this should be bound on demand, only when there are unassigned agents @todo
+
+	if (AwaitingAvatar.Num() == 0 || InActor == nullptr)
+	{
+		return;
+	}
+
+	ensure(AvatarToAgent.Find(HashAvatar(*InActor)) == nullptr);
+
+	// @todo extremely wasteful! needs rethinking/reimplementation
+	TArray<UMLAdapterAgent*> AwaitingAvatarCopy = AwaitingAvatar;
+	AwaitingAvatar.Reset();
+
+	bool bAssigned = false;
+	for (UMLAdapterAgent* Agent : AwaitingAvatarCopy)
+	{
+		if (!bAssigned && Agent->IsSuitableAvatar(*InActor))
+		{
+			BindAvatar(*Agent, *InActor);
+			if (Agent->GetAvatar() == InActor)
+			{
+				bAssigned = true;
+			}
+			else
+			{
+				AwaitingAvatar.Add(Agent);
+			}
+		}
+		else
+		{
+			AwaitingAvatar.Add(Agent);
+		}
+	}		
+}
+
+void UMLAdapterSession::OnPostWorldInit(UWorld& World)
+{
+	SetWorld(&World);
+}
+
+void UMLAdapterSession::OnWorldCleanup(UWorld& World, bool bSessionEnded, bool bCleanupResources)
+{
+	if (CachedWorld == &World)
+	{
+		SetWorld(nullptr);
+	}
+}
+
+void UMLAdapterSession::OnGameModeInitialized(AGameModeBase& GameModeBase)
+{
+	SetGameMode(&GameModeBase);
+}
+
+void UMLAdapterSession::SetGameMode(AGameModeBase* GameModeBase)
+{
+	CachedGameMode = GameModeBase;
+	AGameMode* AsGameMode = Cast<AGameMode>(GameModeBase);
+	if (AsGameMode)
+	{
+		OnGameModeMatchStateSet(AsGameMode->GetMatchState());
+	}
+	else
+	{
+		// a game not utilizing AGameMode's functionality is either a simple game
+		// or a very sophisticated one. In the former case we just assume it's 
+		// 'ready' from the very start:
+		SimulationState = GameModeBase ? EMLAdapterSimState::InProgress : EMLAdapterSimState::Finished;
+		// in the latter case we put it on the user to override this logic.
+	}
+
+	// game-specific data extraction will come here
+}
+
+void UMLAdapterSession::OnGameModePostLogin(AGameModeBase* GameMode, APlayerController* NewPlayer)
+{
+	if (ActorSpawnedDelegateHandle.IsValid() == false && CachedWorld)
+	{
+		FindAvatars(*CachedWorld);
+		ActorSpawnedDelegateHandle = CachedWorld->AddOnActorSpawnedHandler(FOnActorSpawned::FDelegate::CreateUObject(this, &UMLAdapterSession::OnActorSpawned));
+	}
+}
+
+void UMLAdapterSession::OnGameModeMatchStateSet(FName InMatchState)
+{
+	if (InMatchState == MatchState::EnteringMap)
+	{
+		SimulationState = EMLAdapterSimState::BootingUp;
+	}
+	else if (InMatchState == MatchState::WaitingToStart)
+	{
+		SimulationState = EMLAdapterSimState::BootingUp;
+
+		// no point in binding sooner than this 
+		if (ensure(CachedWorld))
+		{
+			FindAvatars(*CachedWorld);
+			if (ActorSpawnedDelegateHandle.IsValid() == false)
+			{
+				ActorSpawnedDelegateHandle = CachedWorld->AddOnActorSpawnedHandler(FOnActorSpawned::FDelegate::CreateUObject(this, &UMLAdapterSession::OnActorSpawned));
+			}
+		}
+	}
+	else if (InMatchState == MatchState::InProgress)
+	{
+		SimulationState = EMLAdapterSimState::InProgress;
+	}
+	else if (InMatchState == MatchState::WaitingPostMatch
+		|| InMatchState == MatchState::LeavingMap
+		|| InMatchState == MatchState::Aborted)
+	{
+		SimulationState = EMLAdapterSimState::Finished;
+	}
+}
+
+void UMLAdapterSession::Open()
+{
+	bActive = true;
+}
+
+void UMLAdapterSession::Close()
+{
+	bActive = false;
+
+	// 'destroy' agents by clearing the async flag and letting the GC clean them
+	for (UMLAdapterAgent* Agent : Agents)
+	{
+		Agent->ClearInternalFlags(EInternalObjectFlags::Async);
+	}
+	Agents.Reset();
+
+	if (CachedWorld && CachedWorld->GetGameInstance())
+	{
+		CachedWorld->GetGameInstance()->GetOnPawnControllerChanged().RemoveAll(this);
+	}
+}
+
+void UMLAdapterSession::Tick(float DeltaTime)
+{
+	LastTimestamp = CachedWorld ? CachedWorld->GetTimeSeconds() : -1.f;
+
+	// @todo for perf reasons we could grab all the agents' senses and tick them
+	// by class to keep the cache hot
+	for (UMLAdapterAgent* Agent : Agents)
+	{
+		Agent->Sense(DeltaTime);
+	}
+
+	for (UMLAdapterAgent* Agent : Agents)
+	{
+		Agent->Think(DeltaTime);
+	}
+
+	for (UMLAdapterAgent* Agent : Agents)
+	{
+		Agent->Act(DeltaTime);
+	}
+
+#if WITH_EDITORONLY_DATA
+	if (CachedWorld && GIntraFrameDebuggingGameThread)
+	{
+		GIntraFrameDebuggingGameThread = false;
+		// the WorldTicker will clear it to allow next tick
+	}
+#endif // WITH_EDITORONLY_DATA
+}
+
+void UMLAdapterSession::ConfigureAsServer()
+{
+
+}
+
+void UMLAdapterSession::ConfigureAsClient()
+{
+
+}
+
+void UMLAdapterSession::ResetWorld(FMLAdapter::FAgentID AgentID)
+{
+	if (CachedGameMode == nullptr)
+	{
+		return;
+	}
+
+	if (AgentID != FMLAdapter::InvalidAgentID)
+	{
+		if (Agents.IsValidIndex(AgentID) && GetAgent(AgentID)->GetAvatar())
+		{
+			AController* AvatarController = FMLAdapter::ActorToController(*GetAgent(AgentID)->GetAvatar());
+			CachedGameMode->RestartPlayer(AvatarController);
+		}
+	}
+	else
+	{
+		for (UMLAdapterAgent* Agent : Agents)
+		{
+			if (Agent->GetAvatar())
+			{
+				AController* AvatarController = FMLAdapter::ActorToController(*Agent->GetAvatar());
+				CachedGameMode->RestartPlayer(AvatarController);
+			}
+		}
+	}
+}
+
+bool UMLAdapterSession::IsDone() const
+{ 
+	return (SimulationState == EMLAdapterSimState::Finished)
+		|| (CachedGameMode != nullptr
+			&& CachedGameMode->HasMatchEnded());
+}
+
+bool UMLAdapterSession::IsReady() const
+{
+	return SimulationState == EMLAdapterSimState::InProgress 
+		&& (CachedGameMode != nullptr)
+		&& (CachedGameMode->HasMatchStarted() == true)
+		&& (CachedGameMode->HasMatchEnded() == false);
+}
+
+FMLAdapter::FAgentID UMLAdapterSession::AddAgent()
+{
+	FScopeLock Lock(&AgentOpCS);
+
+	UClass* AgentClass = UMLAdapterSettings::GetAgentClass().Get()
+		? UMLAdapterSettings::GetAgentClass().Get()
+		: UMLAdapterAgent::StaticClass();
+
+	UMLAdapterAgent* NewAgent = FMLAdapter::NewObject<UMLAdapterAgent>(this, AgentClass);
+
+	NewAgent->SetAgentID(Agents.Add(NewAgent));
+	NewAgent->Configure(NewAgent->GetConfig());
+
+	return NewAgent->GetAgentID();
+}
+
+FMLAdapter::FAgentID UMLAdapterSession::AddAgent(const FMLAdapterAgentConfig& InConfig)
+{
+	FScopeLock Lock(&AgentOpCS);
+
+	UClass* AgentClass = FMLAdapterLibrarian::Get().FindAgentClass(InConfig.AgentClassName);
+	UMLAdapterAgent* NewAgent = FMLAdapter::NewObject<UMLAdapterAgent>(this, AgentClass);
+	NewAgent->SetAgentID(Agents.Add(NewAgent));
+
+	NewAgent->Configure(InConfig);
+
+	return NewAgent->GetAgentID();
+}
+
+FMLAdapter::FAgentID UMLAdapterSession::GetNextAgentID(FMLAdapter::FAgentID ReferenceAgentID) const
+{
+	if (Agents.Num() == 0)
+	{
+		return FMLAdapter::InvalidAgentID;
+	}
+		
+	int Index = (ReferenceAgentID != FMLAdapter::InvalidAgentID)
+		? (int(ReferenceAgentID) + 1) % Agents.Num()
+		: 0;
+
+	for (int Iter = 0; Iter < Agents.Num(); ++Iter)
+	{
+		if (Agents[Index])
+		{
+			return FMLAdapter::FAgentID(Index);
+		}
+		Index = (Index + 1) % Agents.Num();
+	}
+
+	return FMLAdapter::InvalidAgentID;
+}
+
+UMLAdapterAgent* UMLAdapterSession::GetAgent(FMLAdapter::FAgentID AgentID)
+{
+	if (Agents.IsValidIndex(AgentID) == false)
+	{
+		UE_LOG(LogUnrealEditorMLAdapter, Warning, TEXT("UMLAdapterSession::GetAgent: Invalid AgentID. Failing"));
+		return nullptr;
+	}
+
+	return Agents[AgentID];
+}
+
+void UMLAdapterSession::RemoveAgent(FMLAdapter::FAgentID AgentID)
+{
+	if (Agents.IsValidIndex(AgentID) == false)
+	{
+		return;
+	}
+
+	// remove from Agents only if it's the last agent in the list since 
+	// external code refers to agents by ID which is an index to this array
+	// @todo consider switching over to a TMap
+	// @todo consider always nulling out, i.e. not removing elements from Agents
+	UMLAdapterAgent* Agent = Agents[AgentID];
+	if (!ensure(Agent))
+	{
+		return;
+	}
+	
+	OnBeginAgentRemove.Broadcast(*Agent);
+
+	if (Agents.Num() == AgentID)
+	{
+		Agents.Pop(/*bAllowShrinking=*/false);
+	}
+	else
+	{
+		Agents[AgentID] = nullptr;
+	}
+	for (auto It = AvatarToAgent.CreateIterator(); It; ++It)
+	{
+		if (It.Value() == Agent)
+		{
+			It.RemoveCurrent();
+		}
+	}
+	AwaitingAvatar.RemoveSingleSwap(Agent, /*bAllowShrinking=*/false);
+	// there should have been only one agent in AwaitingAvatar
+	ensureMsgf(AwaitingAvatar.Find(Agent) == false, TEXT("there should have been only one agent in AwaitingAvatar"));
+
+	Agent->ClearInternalFlags(EInternalObjectFlags::Async);
+}
+
+bool UMLAdapterSession::IsAgentReady(FMLAdapter::FAgentID AgentID) const
+{
+	return Agents.IsValidIndex(AgentID) && Agents[AgentID] && Agents[AgentID]->IsReady();
+}
+
+void UMLAdapterSession::FindAvatars(UWorld& World)
+{
+	// @todo naive implementation for now, subject to optimization in the future
+	TArray<UMLAdapterAgent*> AwaitingAvatarCopy = AwaitingAvatar;
+	AwaitingAvatar.Reset();
+
+	for (UMLAdapterAgent* Agent : AwaitingAvatarCopy)
+	{
+		if (ensure(Agent))
+		{
+			ensure(Agent->GetAvatar() == nullptr); // if not then the avatar has been assigned outside of normal procedure 
+			RequestAvatarForAgent(*Agent, &World);
+		}
+	}
+}
+
+void UMLAdapterSession::RemoveAvatars(UWorld* World)
+{
+	for (UMLAdapterAgent* Agent : Agents)
+	{
+		if (Agent && Agent->GetAvatar() && (World == nullptr || Agent->GetAvatar()->GetWorld() == World))
+		{
+			AActor* OldAvatar = Agent->GetAvatar();
+			Agent->SetAvatar(nullptr);
+			OnAgentAvatarChanged.Broadcast(*Agent, OldAvatar);
+		}
+	}
+}
+
+bool UMLAdapterSession::RequestAvatarForAgent(FMLAdapter::FAgentID& AgentID, UWorld* InWorld)
+{
+	UMLAdapterAgent* Agent = GetAgent(AgentID);
+	return (Agent != nullptr) && RequestAvatarForAgent(*Agent, InWorld);
+}
+
+bool UMLAdapterSession::RequestAvatarForAgent(UMLAdapterAgent& Agent, UWorld* InWorld, const bool bForceSearch)
+{
+	if (Agent.GetAvatar() != nullptr)
+	{
+		// skipping.
+		UE_LOG(LogUnrealEditorMLAdapter, Verbose, TEXT("UMLAdapterSession::RequestAvatarForAgent called for agent [%s] while it still has an avatar [%s]. Call ClearAvatar first to null-out agent\'s avatar.")
+			, Agent.GetAgentID(), *GetNameSafe(Agent.GetAvatar()));
+		return false;
+	}
+	if (bForceSearch == false && AwaitingAvatar.Find(&Agent) != INDEX_NONE)
+	{
+		// already waiting, skip
+		return false;
+	}
+
+	// we're adding to awaiting list first to avoid calling RequestAvatarForAgent 
+	// couple of times in sequence
+	AwaitingAvatar.AddUnique(&Agent);
+
+	InWorld = InWorld ? InWorld : CachedWorld;
+	
+	if (InWorld == nullptr)
+	{
+		UE_LOG(LogUnrealEditorMLAdapter, Warning, TEXT("UMLAdapterSession::RequestAvatarForAgent called with InWorld and CachedWorld both being null. Auto-failure."));
+		return false;
+	}
+
+	AActor* Avatar = nullptr;
+	if (InWorld != nullptr)
+	{
+		// @todo might want to make special cases for Controllers and 
+		UClass* AvatarClass = Agent.GetConfig().AvatarClass;
+		if (ensure(AvatarClass))
+		{
+			for (TActorIterator<AActor> It(InWorld, AvatarClass); It; ++It)
+			{
+				if (Agent.IsSuitableAvatar(**It)
+					// AND not already an avatar
+					&& AvatarToAgent.Find(HashAvatar(**It)) == nullptr)
+				{
+					Avatar = *It;
+					break;
+				}
+			}
+		}
+	}
+
+	if (Avatar)
+	{
+		BindAvatar(Agent, *Avatar);
+		ensureMsgf(Agent.GetAvatar() == Avatar, TEXT("If we get here and the avatar setting fails it means the above process, leading to avatar selection was flawed"));
+	}
+
+	return Agent.GetAvatar() && (Agent.GetAvatar() == Avatar);
+}
+
+void UMLAdapterSession::BindAvatar(UMLAdapterAgent& Agent, AActor& Avatar)
+{
+	AActor* OldAvatar = Agent.GetAvatar();
+	ClearAvatar(Agent);
+
+	Agent.SetAvatar(&Avatar);
+	AwaitingAvatar.RemoveSingleSwap(&Agent, /*bAllowShrinking=*/false);	
+	AvatarToAgent.Add(HashAvatar(Avatar), &Agent);
+
+	OnAgentAvatarChanged.Broadcast(Agent, OldAvatar);
+}
+
+void UMLAdapterSession::ClearAvatar(UMLAdapterAgent& Agent)
+{
+	AActor* OldAvatar = Agent.GetAvatar();
+	if (OldAvatar == nullptr)
+	{
+		// it's possible the previous avatar is already gone. We need to look
+		// through the map values
+		const uint32* Key = AvatarToAgent.FindKey(&Agent);
+		if (Key)
+		{
+			AvatarToAgent.Remove(*Key);
+		}
+		return;
+	}
+
+	UMLAdapterAgent* BoundAgent = nullptr;
+	AvatarToAgent.RemoveAndCopyValue(HashAvatar(*OldAvatar), BoundAgent);
+	ensure(BoundAgent == &Agent);
+	// @todo what if it causes another RequestAvatarForAgent call
+	// should probably ignore the second one
+	Agent.SetAvatar(nullptr);
+	OnAgentAvatarChanged.Broadcast(Agent, OldAvatar);
+}
+
+const UMLAdapterAgent* UMLAdapterSession::FindAgentByAvatar(AActor& Avatar) const
+{
+	for (UMLAdapterAgent* Agent : Agents)
+	{
+		if (Agent && Agent->GetAvatar() == &Avatar)
+		{
+			return Agent;
+		}
+	}
+
+	return nullptr;
+}
+
+#if WITH_GAMEPLAY_DEBUGGER
+#include "GameplayDebuggerCategory.h"
+
+void UMLAdapterSession::DescribeSelfToGameplayDebugger(FGameplayDebuggerCategory& DebuggerCategory) const
+{
+	if (Agents.Num() > 0)
+	{
+		bool bInvalidAgents = false;
+		FString ActiveAgentIDs;
+		for (UMLAdapterAgent* Agent : Agents)
+		{
+			if (Agent)
+			{
+				ActiveAgentIDs += FString::Printf(TEXT("{%s}%d,")
+					, AwaitingAvatar.Find(Agent) != INDEX_NONE ? TEXT("grey") : TEXT("white")
+					, Agent->GetAgentID());
+			}
+			else
+			{
+				bInvalidAgents = true;
+			}
+		}
+		
+		DebuggerCategory.AddTextLine(FString::Printf(TEXT("{green}Active agents: %s"), *ActiveAgentIDs));
+
+		if (bInvalidAgents || AwaitingAvatar.Find(nullptr) != INDEX_NONE)
+		{
+			DebuggerCategory.AddTextLine(FString::Printf(TEXT("{red} invalid agents found!")));
+		}
+	}
+}
+#endif // WITH_GAMEPLAY_DEBUGGER
diff --git a/Engine/Plugins/AI/MLAdapter/Source/MLAdapter/Public/MLAdapterManager.h b/Engine/Plugins/AI/MLAdapter/Source/MLAdapter/Public/MLAdapterManager.h
deleted file mode 100644
diff --git a/Engine/Plugins/AI/MLAdapter/Source/MLAdapter/Public/MLAdapterSession.h b/Engine/Plugins/AI/MLAdapter/Source/MLAdapter/Public/MLAdapterSession.h
deleted file mode 100644
diff --git a/Engine/Plugins/AI/MLAdapter/Source/MLAdapter/Public/MLAdapterSettings.h b/Engine/Plugins/AI/MLAdapter/Source/MLAdapter/Public/MLAdapterSettings.h
--- a/Engine/Plugins/AI/MLAdapter/Source/MLAdapter/Public/MLAdapterSettings.h
+++ b/Engine/Plugins/AI/MLAdapter/Source/MLAdapter/Public/MLAdapterSettings.h
@@ -4,8 +4,8 @@
 
 #include "Engine/DeveloperSettings.h"
 #include "Agents/MLAdapterAgent.h"
-#include "MLAdapterSession.h"
-#include "MLAdapterManager.h"
+#include "Sessions/MLAdapterSession.h"
+#include "Managers/MLAdapterManager.h"
 #include "MLAdapterSettings.generated.h"
 
 
@@ -40,7 +40,7 @@ protected:
 
 	UPROPERTY(EditAnywhere, config, Category = MLAdapter, meta = (MetaClass = "MLAdapterSession"))
 	FSoftClassPath SessionClass;
-	
+
 	UPROPERTY(EditAnywhere, config, Category = MLAdapter, meta = (MetaClass = "MLAdapterAgent"))
 	FSoftClassPath DefaultAgentClass;
 
diff --git a/Engine/Plugins/AI/MLAdapter/Source/MLAdapter/Public/Managers/MLAdapterManager.h b/Engine/Plugins/AI/MLAdapter/Source/MLAdapter/Public/Managers/MLAdapterManager.h
new file mode 100644
--- /dev/null
+++ b/Engine/Plugins/AI/MLAdapter/Source/MLAdapter/Public/Managers/MLAdapterManager.h
@@ -0,0 +1,190 @@
+// Copyright Epic Games, Inc. All Rights Reserved.
+
+#pragma once
+
+#include "CoreMinimal.h"
+#include "Delegates/DelegateCombinations.h"
+#include "Delegates/IDelegateInstance.h"
+#include "Logging/LogMacros.h"
+#include "Tickable.h"
+#include "Engine/World.h"
+#include "Sessions/MLAdapterSession.h"
+#include "MLAdapterLibrarian.h"
+#include "MLAdapterManager.generated.h"
+
+namespace rpc { class server; }
+namespace FMLAdapterConsoleCommands { struct FHelper; }
+class UMLAdapterSession;
+class UMLAdapterActuator;
+class UMLAdapterSensor;
+class UMLAdapterAgent;
+class AGameModeBase;
+class APlayerController;
+class UWorld;
+
+UENUM()
+enum class EMLAdapterServerMode
+{
+	Invalid,
+	Server,
+	Client,
+	Standalone, // this applies both to Standalone games as well as PIE
+	AutoDetect,
+};
+
+
+UCLASS(Transient)
+class MLADAPTER_API UMLAdapterManager : public UObject, public FTickableGameObject, public FSelfRegisteringExec
+{
+	GENERATED_BODY()
+public:
+	using FRPCFunctionBind = TFunction<void(FRPCServer&/*RPCServer*/)>;
+	DECLARE_MULTICAST_DELEGATE_OneParam(FOnGenericRPCServerDelegate, FRPCServer& /*Server*/);
+	DECLARE_MULTICAST_DELEGATE(FOnGenericEvent);
+	
+	UMLAdapterManager(const FObjectInitializer& ObjectInitializer = FObjectInitializer::Get());
+	virtual void BeginDestroy() override;
+	virtual void PostInitProperties() override;
+
+	// FTickableGameObject begin
+	virtual void Tick(float DeltaTime) override;
+	virtual TStatId GetStatId() const override;
+	virtual bool IsTickableWhenPaused() const override { return true; }
+	virtual bool IsTickableInEditor() const override { return true; }
+	virtual ETickableTickType GetTickableTickType() const override;
+	virtual bool IsTickable() const override;
+	// FTickableGameObject end
+
+	// FExec begin
+	virtual bool Exec(class UWorld* InWorld, const TCHAR* Cmd, FOutputDevice& Ar) override;
+	// FExec end
+
+	virtual void BindToDelegates();
+	virtual void CleanUpDelegates();
+
+	virtual void OnPostWorldInit(UWorld* World, const UWorld::InitializationValues);
+	virtual void OnWorldCleanup(UWorld* World, bool bSessionEnded, bool bCleanupResources);
+	virtual void OnBeginPIE(const bool bIsSimulating);
+	virtual void OnEndPIE(const bool bIsSimulating);
+	virtual void OnGameModeInitialized(AGameModeBase* GameMode);
+	virtual void OnGameModePostLogin(AGameModeBase* GameMode, APlayerController* NewPlayer);
+	/** @note that this might not get called at all if the project's game mode 
+	 *	doesn't extend AGameMode*/
+	virtual void OnGameModeMatchStateSet(FName MatchState);
+
+	/** If a server is already running it will be shut down before the new instance gets created
+	 *	@param ServerThreads best set at the number of external clients that are going to be connecting */
+	virtual void StartServer(uint16 Port, EMLAdapterServerMode InMode = EMLAdapterServerMode::AutoDetect, uint16 ServerThreads = 1);
+	virtual void StopServer();
+	virtual bool IsRunning() const;
+
+	virtual void ConfigureAsServer(FRPCServer& Server);
+	/** "Client" in this context means UnrealEngine game client, not RPC client */
+	virtual void ConfigureAsClient(FRPCServer& Server);
+	/** Essentially calls both the server and client versions */
+	virtual void ConfigureAsStandalone(FRPCServer& Server);
+
+	/** if given World doesn't have an AI system this call results in creating one */
+	virtual void EnsureAISystemPresence(UWorld& World);
+
+	/** if given World doesn't have a Navigation system instance this call results in creating one */
+	virtual void EnsureNavigationSystemPresence(UWorld& World);
+
+	virtual UMLAdapterSession* CreateNewSession();
+	virtual void SetSession(UMLAdapterSession* NewSession);
+	virtual void CloseSession(UMLAdapterSession& InSession);
+
+	/** Returns current session. If one doesn't exist, it gets created. */
+	virtual UMLAdapterSession& GetSession();
+	bool HasSession() const { return IsValid(Session) && (Session->IsUnreachable() == false); }
+
+	void RegisterSensorClass(const TSubclassOf<UMLAdapterSensor>& Class) { Librarian.RegisterSensorClass(Class); }
+	void RegisterActuatorClass(const TSubclassOf<UMLAdapterActuator>& Class) { Librarian.RegisterActuatorClass(Class); }
+	void RegisterAgentClass(const TSubclassOf<UMLAdapterAgent>& Class) { Librarian.RegisterAgentClass(Class); }
+
+	virtual void ResetWorld();
+	void SetManualWorldTickEnabled(bool bEnable);
+
+	FOnGenericRPCServerDelegate& GetOnAddClientFunctions() { return OnAddClientFunctions; }
+	FOnGenericRPCServerDelegate& GetOnAddServerFunctions() { return OnAddServerFunctions; }
+
+	FORCEINLINE static UMLAdapterManager& Get();
+	FORCEINLINE static bool IsReady();
+
+	const FMLAdapterLibrarian& GetLibrarian() const { return Librarian; }
+
+	bool IsWorldRealTime() const { return (bTickWorldManually == false); }
+
+	FOnGenericEvent& GetOnCurrentSessionChanged() { return OnCurrentSessionChanged; }
+
+protected:
+
+	void AddCommonFunctions(FRPCServer& Server);
+
+	virtual bool ShouldInitForWorld(const UWorld& World) const;
+
+protected:	
+	friend struct FMLAdapterConsoleCommands::FHelper;
+
+	UPROPERTY()
+	UMLAdapterSession* Session;
+
+	UPROPERTY()
+	UWorld* LastActiveWorld;
+
+	UPROPERTY()
+	FMLAdapterLibrarian Librarian;
+
+	FOnGenericRPCServerDelegate OnAddClientFunctions;
+	FOnGenericRPCServerDelegate OnAddServerFunctions;
+
+	FOnGenericEvent OnCurrentSessionChanged;
+
+	EMLAdapterServerMode RequestedFunctionMode = EMLAdapterServerMode::Invalid;
+	EMLAdapterServerMode CurrentFunctionMode = EMLAdapterServerMode::Invalid;
+	uint16 CurrentPort = 0;
+	uint16 CurrentServerThreads = 1;
+
+	TArray<uint8> Data;
+
+	float WorldFPS = 20.f;
+
+	uint32 bCommonFunctionsAdded : 1;
+	uint32 bTickWorldManually : 1;
+	
+	/** is the manager is in 'manual ticking mode' (where external client is 
+	 *	responsible for progressing the world sim by calling 'request_world_tick' 
+	 *	function) the simulation will progress by StepsRequested ticks before pausing */
+	int32 StepsRequested = 0;
+
+	FDelegateHandle OnPostWorldInitializationHandle;
+	FDelegateHandle OnWorldCleanupHandle;
+	FDelegateHandle OnGameModeInitializedHandle;
+	FDelegateHandle OnGameModePostLoginHandle;
+	FDelegateHandle OnGameModeMatchStateSetHandle;
+#if WITH_EDITORONLY_DATA
+	FDelegateHandle BeginPIEHandle;
+	FDelegateHandle EndPIEHandle;
+#endif // WITH_EDITORONLY_DATA
+
+	static UMLAdapterManager* ManagerInstance;
+public:
+	static void RecreateManagerInstance();
+
+	static FOnGenericEvent OnPostInit;
+};
+
+//----------------------------------------------------------------------//
+// inlines 
+//----------------------------------------------------------------------//
+UMLAdapterManager& UMLAdapterManager::Get()
+{ 
+	// the only way for this check to fail is to call it too soon.
+	check(ManagerInstance);  
+	return *ManagerInstance; 
+}
+
+bool UMLAdapterManager::IsReady()
+{ 
+	return (ManagerInstance != nullptr); 
+}
\ No newline at end of file
diff --git a/Engine/Plugins/AI/MLAdapter/Source/MLAdapter/Public/Sessions/MLAdapterSession.h b/Engine/Plugins/AI/MLAdapter/Source/MLAdapter/Public/Sessions/MLAdapterSession.h
new file mode 100644
--- /dev/null
+++ b/Engine/Plugins/AI/MLAdapter/Source/MLAdapter/Public/Sessions/MLAdapterSession.h
@@ -0,0 +1,184 @@
+// Copyright Epic Games, Inc. All Rights Reserved.
+
+#pragma once 
+
+#include "CoreMinimal.h"
+#include "UObject/ObjectMacros.h"
+#include "UObject/UObjectGlobals.h"
+#include "UObject/Object.h"
+#include "MLAdapterTypes.h"
+#include "Tickable.h"
+#include "Engine/World.h"
+#include "GameFramework/Actor.h"
+#include "MLAdapterSession.generated.h"
+
+
+class AGameModeBase;
+class UMLAdapterAgent;
+struct FMLAdapterAgentConfig;
+class APawn;
+class AController;
+class APlayerController;
+class UGameInstance;
+
+
+enum class EMLAdapterSimState : uint8
+{
+	BootingUp,
+	InProgress,
+	Finished,
+};
+
+
+UCLASS()
+class MLADAPTER_API UMLAdapterSession : public UObject
+{
+	GENERATED_BODY()
+public:
+	DECLARE_MULTICAST_DELEGATE_TwoParams(FOnAgentAvatarChangedDelegate, UMLAdapterAgent& /*Agent*/, AActor* /*OldAvatar*/);
+	DECLARE_MULTICAST_DELEGATE_OneParam(FOnBeginAgentRemove, UMLAdapterAgent& /*Agent*/);
+
+	virtual UWorld* GetWorld() const override { return CachedWorld; }
+	virtual void PostInitProperties() override;
+
+	virtual void BeginDestroy() override;
+
+	// @todo this needs further consideration.
+
+	UGameInstance* GetGameInstance() const { return CachedWorld ? CachedWorld->GetGameInstance() : nullptr; }
+	/** NewWorld might be null */
+	virtual void SetWorld(UWorld* NewWorld);
+
+	virtual void OnActorSpawned(AActor* InActor);
+	virtual void OnPostWorldInit(UWorld& World);
+	virtual void OnWorldCleanup(UWorld& World, bool bSessionEnded, bool bCleanupResources);
+	virtual void OnGameModeInitialized(AGameModeBase& GameModeBase);
+	virtual void OnGameModeMatchStateSet(FName InMatchState);
+	virtual void OnGameModePostLogin(AGameModeBase* GameMode, APlayerController* NewPlayer);
+
+	virtual void Open();
+	virtual void Close();
+
+	virtual void Tick(float DeltaTime);
+
+	/** This is where Session can add Authority-side-specific functions by calling 
+	 *	UMLAdapterManager::Get().AddServerFunctionBind */
+	virtual void ConfigureAsServer();
+
+	/** This is where Session can add Client-side-specific functions by calling 
+	 *	UMLAdapterManager::Get().AddClientFunctionBind. "Client" in this context means 
+	 *	UnrealEngine game client, not RPC client */
+	virtual void ConfigureAsClient();
+
+	/** using FMLAdapter::InvalidAgentID for AgentID will reset all agents */
+	virtual void ResetWorld(FMLAdapter::FAgentID AgentID = FMLAdapter::InvalidAgentID);
+
+	bool IsDone() const;
+	bool IsReady() const;
+
+	float GetTimestamp() const { return LastTimestamp; }
+
+	void SetManualWorldTickEnabled(bool bEnable);
+
+	FOnAgentAvatarChangedDelegate GetOnAgentAvatarChanged() { return OnAgentAvatarChanged; }
+	FOnBeginAgentRemove GetOnBeginAgentRemove() { return OnBeginAgentRemove; }
+
+	//----------------------------------------------------------------------//
+	// Agent/Avatar management 
+	//----------------------------------------------------------------------//
+	FMLAdapter::FAgentID AddAgent();
+	FMLAdapter::FAgentID AddAgent(const FMLAdapterAgentConfig& InConfig);
+
+	/** @return next valid agent ID. Note that the return value might be equal to 
+	 *	ReferenceAgentID if there's only one agent. Will be FMLAdapter::InvalidAgentId 
+	 *	if no agents registered */
+	FMLAdapter::FAgentID GetNextAgentID(FMLAdapter::FAgentID ReferenceAgentID) const;
+	UMLAdapterAgent* GetAgent(FMLAdapter::FAgentID AgentID);
+
+	void RemoveAgent(FMLAdapter::FAgentID AgentID);
+
+	bool IsAgentReady(FMLAdapter::FAgentID AgentID) const;
+
+	/** Finds avatar in given World for every avatar-less agent in AwaitingAvatar */
+	void FindAvatars(UWorld& World);
+
+	/** Processes Agents and removes all agent avatars belonging to World. 
+	 *	If World is null the function will remove all avatars */
+	void RemoveAvatars(UWorld* World);
+
+	/** Finds a suitable avatar in InWorld (or CachedWorld, if InWorld is null) 
+	 *	for given agent, as specified by FMLAdapterAgentConfig.AvatarClass
+	 *	and confirmed by Agent->IsSuitableAvatar call. If no suitable avatar is 
+	 *	found this agent will be added to "waiting list" (AwaitingAvatar)
+	 *	@param bForceSearch if true will ignore whether the Agent is already waiting in 
+	 *		AwaitingAvatar and will perform the search right away. Note that the Agent 
+	 *		might still end up in AwaitingAvatar if there's no suitable avatars available 
+	 *	@return true if an avatar has been assigned. False otherwise.*/
+	virtual bool RequestAvatarForAgent(UMLAdapterAgent& Agent, UWorld* InWorld = nullptr, const bool bForceSearch = false);
+
+	bool RequestAvatarForAgent(FMLAdapter::FAgentID& AgentID, UWorld* InWorld = nullptr);
+
+	void BindAvatar(UMLAdapterAgent& Agent, AActor& Avatar);
+	void ClearAvatar(UMLAdapterAgent& Agent);
+
+	int32 GetAgentsCount() const { return Agents.Num(); }
+
+	const UMLAdapterAgent* FindAgentByAvatar(AActor& Avatar) const;
+
+	//----------------------------------------------------------------------//
+	// debug 
+	//----------------------------------------------------------------------//
+#if WITH_GAMEPLAY_DEBUGGER
+	void DescribeSelfToGameplayDebugger(class FGameplayDebuggerCategory& DebuggerCategory) const;
+#endif // WITH_GAMEPLAY_DEBUGGER
+
+protected:
+	virtual void SetGameMode(AGameModeBase* GameModeBase);
+
+	FORCEINLINE static uint32 HashAvatar(const AActor& Avatar)
+	{
+		return Avatar.GetUniqueID();
+	}
+
+	UPROPERTY()
+	AGameModeBase* CachedGameMode;
+
+	UPROPERTY()
+	UWorld* CachedWorld;
+
+	/** @see HashAvatar */
+	UPROPERTY()
+	TMap<uint32, UMLAdapterAgent*> AvatarToAgent;
+
+	UPROPERTY()
+	TArray<UMLAdapterAgent*> Agents;
+
+	UPROPERTY()
+	TArray<UMLAdapterAgent*> AwaitingAvatar;
+
+	FOnAgentAvatarChangedDelegate OnAgentAvatarChanged;
+	FOnBeginAgentRemove OnBeginAgentRemove;
+
+	FDelegateHandle ActorSpawnedDelegateHandle;
+	
+	EMLAdapterSimState SimulationState;
+
+	float LastTimestamp = -1.f;
+
+	bool bActive = false;
+	bool bTickWorldManually = false;
+
+	mutable FCriticalSection AgentOpCS;
+
+	struct FWorldTicker : public FTickableGameObject
+	{
+		TWeakObjectPtr<UWorld> CachedWorld;
+		FWorldTicker(UWorld* InWorld) : CachedWorld(InWorld) {}
+		virtual ~FWorldTicker();
+		virtual void Tick(float DeltaTime) override;
+		virtual UWorld* GetTickableGameObjectWorld() const override { return CachedWorld.Get(); }
+		virtual TStatId GetStatId() const { return TStatId(); }
+	};
+
+	TSharedPtr<FWorldTicker> WorldTicker;
+};
diff --git a/Engine/Plugins/AI/MLAdapter/Source/MLAdapterTestSuite/Private/AgentTest.cpp b/Engine/Plugins/AI/MLAdapter/Source/MLAdapterTestSuite/Private/AgentTest.cpp
--- a/Engine/Plugins/AI/MLAdapter/Source/MLAdapterTestSuite/Private/AgentTest.cpp
+++ b/Engine/Plugins/AI/MLAdapter/Source/MLAdapterTestSuite/Private/AgentTest.cpp
@@ -4,7 +4,7 @@
 #include "Agents/MLAdapterAgent.h"
 #include "AIController.h"
 #include "GameFramework/Pawn.h"
-#include "MLAdapterSession.h"
+#include "Sessions/MLAdapterSession.h"
 
 #define LOCTEXT_NAMESPACE "AITestSuite_MLAdapterTest"
 
diff --git a/Engine/Plugins/AI/MLAdapter/Source/MLAdapterTestSuite/Private/RPCTestBase.h b/Engine/Plugins/AI/MLAdapter/Source/MLAdapterTestSuite/Private/RPCTestBase.h
--- a/Engine/Plugins/AI/MLAdapter/Source/MLAdapterTestSuite/Private/RPCTestBase.h
+++ b/Engine/Plugins/AI/MLAdapter/Source/MLAdapterTestSuite/Private/RPCTestBase.h
@@ -5,7 +5,7 @@
 #include "CoreMinimal.h"
 #include "AITestsCommon.h"
 #include "Misc/App.h"
-#include "MLAdapterManager.h"
+#include "Managers/MLAdapterManager.h"
 
 
 #if WITH_RPCLIB
diff --git a/Engine/Plugins/AI/MLAdapter/Source/MLAdapterTestSuite/Private/SessionTest.cpp b/Engine/Plugins/AI/MLAdapter/Source/MLAdapterTestSuite/Private/SessionTest.cpp
--- a/Engine/Plugins/AI/MLAdapter/Source/MLAdapterTestSuite/Private/SessionTest.cpp
+++ b/Engine/Plugins/AI/MLAdapter/Source/MLAdapterTestSuite/Private/SessionTest.cpp
@@ -1,8 +1,8 @@
 // Copyright Epic Games, Inc. All Rights Reserved.
 
 #include "AITestsCommon.h"
-#include "MLAdapterManager.h"
-#include "MLAdapterSession.h"
+#include "Managers/MLAdapterManager.h"
+#include "Sessions/MLAdapterSession.h"
 #include "Agents/MLAdapterAgent.h"
 #include "AIController.h"
 #include "GameFramework/Pawn.h"
diff --git a/Engine/Plugins/Experimental/MLAdapterInference/Source/MLAdapterInference/Private/MLAdapterInferenceManager.cpp b/Engine/Plugins/Experimental/MLAdapterInference/Source/MLAdapterInference/Private/MLAdapterInferenceManager.cpp
--- a/Engine/Plugins/Experimental/MLAdapterInference/Source/MLAdapterInference/Private/MLAdapterInferenceManager.cpp
+++ b/Engine/Plugins/Experimental/MLAdapterInference/Source/MLAdapterInference/Private/MLAdapterInferenceManager.cpp
@@ -5,7 +5,7 @@
 #include "MLAdapterInferenceSettings.h"
 #include "MLAdapterInferenceTypes.h"
 #include "NeuralNetwork.h"
-#include "MLAdapterSession.h"
+#include "Sessions/MLAdapterSession.h"
 
 
 #if WITH_EDITORONLY_DATA
diff --git a/Engine/Plugins/Experimental/MLAdapterInference/Source/MLAdapterInference/Public/MLAdapterInferenceManager.h b/Engine/Plugins/Experimental/MLAdapterInference/Source/MLAdapterInference/Public/MLAdapterInferenceManager.h
--- a/Engine/Plugins/Experimental/MLAdapterInference/Source/MLAdapterInference/Public/MLAdapterInferenceManager.h
+++ b/Engine/Plugins/Experimental/MLAdapterInference/Source/MLAdapterInference/Public/MLAdapterInferenceManager.h
@@ -3,7 +3,7 @@
 #pragma once 
 
 #include "CoreMinimal.h"
-#include "MLAdapterManager.h"
+#include "Managers/MLAdapterManager.h"
 #include "MLAdapterInferenceManager.generated.h"
 
 /**
