MLAdapter: Add new manager and session to make data collection on the client easier

#rnx

diff --git a/Engine/Plugins/AI/MLAdapter/Source/MLAdapter/Private/Managers/MLAdapterManager.cpp b/Engine/Plugins/AI/MLAdapter/Source/MLAdapter/Private/Managers/MLAdapterManager.cpp
--- a/Engine/Plugins/AI/MLAdapter/Source/MLAdapter/Private/Managers/MLAdapterManager.cpp
+++ b/Engine/Plugins/AI/MLAdapter/Source/MLAdapter/Private/Managers/MLAdapterManager.cpp
@@ -274,10 +274,12 @@ void UMLAdapterManager::Tick(float DeltaTime)
 
 UMLAdapterSession* UMLAdapterManager::CreateNewSession()
 {
-	UClass* Class = UMLAdapterSettings::GetSessionClass().Get()
+	UClass* Class = UMLAdapterSettings::GetSessionClass().Get() 
 		? UMLAdapterSettings::GetSessionClass().Get()
 		: UMLAdapterSession::StaticClass();
 
+	UE_LOG(LogUnrealEditorMLAdapter, Log, TEXT("Creating MLAdapter session of class %s"), *GetNameSafe(Class));
+
 	UMLAdapterSession* NewSession = FMLAdapter::NewObject<UMLAdapterSession>(this, Class);
 	check(NewSession);
 	NewSession->SetWorld(LastActiveWorld);
diff --git a/Engine/Plugins/AI/MLAdapter/Source/MLAdapter/Private/Managers/MLAdapterNoRPCManager.cpp b/Engine/Plugins/AI/MLAdapter/Source/MLAdapter/Private/Managers/MLAdapterNoRPCManager.cpp
new file mode 100644
--- /dev/null
+++ b/Engine/Plugins/AI/MLAdapter/Source/MLAdapter/Private/Managers/MLAdapterNoRPCManager.cpp
@@ -0,0 +1,18 @@
+// Copyright Epic Games, Inc. All Rights Reserved.
+
+#include "Managers/MLAdapterNoRPCManager.h"
+#include "Sessions/MLAdapterSession.h"
+
+void UMLAdapterNoRPCManager::OnPostWorldInit(UWorld* World, const UWorld::InitializationValues)
+{
+	if (World && ShouldInitForWorld(*World))
+	{
+		LastActiveWorld = World;
+
+		// GetSession will create the initial session
+		GetSession().OnPostWorldInit(*World);
+
+		// AddAgent will add the default agent from the settings
+		FMLAdapter::FAgentID AgentID = GetSession().AddAgent();
+	}
+}
diff --git a/Engine/Plugins/AI/MLAdapter/Source/MLAdapter/Private/Sessions/MLAdapterLocalDataCollectionSession.cpp b/Engine/Plugins/AI/MLAdapter/Source/MLAdapter/Private/Sessions/MLAdapterLocalDataCollectionSession.cpp
new file mode 100644
--- /dev/null
+++ b/Engine/Plugins/AI/MLAdapter/Source/MLAdapter/Private/Sessions/MLAdapterLocalDataCollectionSession.cpp
@@ -0,0 +1,59 @@
+// Copyright Epic Games, Inc. All Rights Reserved.
+
+#include "Sessions/MLAdapterLocalDataCollectionSession.h"
+#include "Agents/MLAdapterAgent.h"
+#include "MLAdapterTypes.h"
+#include "Misc/FileHelper.h"
+#include "HAL/FileManager.h"
+
+void UMLAdapterLocalDataCollectionSession::OnPostWorldInit(UWorld& World)
+{
+	Super::OnPostWorldInit(World);
+
+	UGameInstance* GameInstance = GetGameInstance();
+	if (GameInstance)
+	{
+		// This will be cleaned up in UMLAdapterLocalDataCollectionSession::Close()
+		GameInstance->GetOnPawnControllerChanged().AddDynamic(this, &UMLAdapterLocalDataCollectionSession::OnPawnControllerChanged);
+	}
+}
+
+void UMLAdapterLocalDataCollectionSession::OnPawnControllerChanged(APawn* InPawn, AController* InController)
+{
+	for (UMLAdapterAgent* Agent : Agents)
+	{
+		if (Agent->GetController() == InController)
+		{
+			PlayerControlledAgent = Agent;
+			break;
+		}
+	}
+}
+
+void UMLAdapterLocalDataCollectionSession::Tick(float DeltaTime)
+{
+	UMLAdapterAgent* Agent = PlayerControlledAgent.Get();
+
+	if (Agent ==  nullptr)
+	{
+		UE_LOG(LogUnrealEditorMLAdapter, Log, TEXT("LocalDataCollectionSession: Player-controlled agent not found yet."))
+		return;
+	}
+
+	Agent->Sense(DeltaTime);
+	TArray<uint8> Buffer;
+	FMLAdapterMemoryWriter Writer(Buffer);
+	Agent->GetObservations(Writer);
+
+	FFileHelper::SaveArrayToFile(Buffer, *FileName, &IFileManager::Get(), EFileWrite::FILEWRITE_Append);
+}
+
+void UMLAdapterLocalDataCollectionSession::Close()
+{
+	Super::Close();
+
+	if (CachedWorld && CachedWorld->GetGameInstance())
+	{
+		CachedWorld->GetGameInstance()->GetOnPawnControllerChanged().RemoveAll(this);
+	}
+}
diff --git a/Engine/Plugins/AI/MLAdapter/Source/MLAdapter/Private/Sessions/MLAdapterSession.cpp b/Engine/Plugins/AI/MLAdapter/Source/MLAdapter/Private/Sessions/MLAdapterSession.cpp
--- a/Engine/Plugins/AI/MLAdapter/Source/MLAdapter/Private/Sessions/MLAdapterSession.cpp
+++ b/Engine/Plugins/AI/MLAdapter/Source/MLAdapter/Private/Sessions/MLAdapterSession.cpp
@@ -231,12 +231,8 @@ void UMLAdapterSession::Close()
 	{
 		Agent->ClearInternalFlags(EInternalObjectFlags::Async);
 	}
-	Agents.Reset();
 
-	if (CachedWorld && CachedWorld->GetGameInstance())
-	{
-		CachedWorld->GetGameInstance()->GetOnPawnControllerChanged().RemoveAll(this);
-	}
+	Agents.Reset();
 }
 
 void UMLAdapterSession::Tick(float DeltaTime)
@@ -330,6 +326,8 @@ FMLAdapter::FAgentID UMLAdapterSession::AddAgent()
 		? UMLAdapterSettings::GetAgentClass().Get()
 		: UMLAdapterAgent::StaticClass();
 
+	UE_LOG(LogUnrealEditorMLAdapter, Log, TEXT("Creating MLAdapter agent of class %s"), *GetNameSafe(AgentClass));
+
 	UMLAdapterAgent* NewAgent = FMLAdapter::NewObject<UMLAdapterAgent>(this, AgentClass);
 
 	NewAgent->SetAgentID(Agents.Add(NewAgent));
diff --git a/Engine/Plugins/AI/MLAdapter/Source/MLAdapter/Public/Managers/MLAdapterNoRPCManager.h b/Engine/Plugins/AI/MLAdapter/Source/MLAdapter/Public/Managers/MLAdapterNoRPCManager.h
new file mode 100644
--- /dev/null
+++ b/Engine/Plugins/AI/MLAdapter/Source/MLAdapter/Public/Managers/MLAdapterNoRPCManager.h
@@ -0,0 +1,18 @@
+// Copyright Epic Games, Inc. All Rights Reserved.
+
+#pragma once 
+
+#include "CoreMinimal.h"
+#include "Managers/MLAdapterManager.h"
+#include "MLAdapterNoRPCManager.generated.h"
+
+/**
+ * No RPC manager won't start RPC server and will immediately start a session and spawn the default agent.
+ */
+UCLASS()
+class MLADAPTER_API UMLAdapterNoRPCManager : public UMLAdapterManager
+{
+	GENERATED_BODY()
+public:
+	virtual void OnPostWorldInit(UWorld* World, const UWorld::InitializationValues) override;
+};
diff --git a/Engine/Plugins/AI/MLAdapter/Source/MLAdapter/Public/Sessions/MLAdapterLocalDataCollectionSession.h b/Engine/Plugins/AI/MLAdapter/Source/MLAdapter/Public/Sessions/MLAdapterLocalDataCollectionSession.h
new file mode 100644
--- /dev/null
+++ b/Engine/Plugins/AI/MLAdapter/Source/MLAdapter/Public/Sessions/MLAdapterLocalDataCollectionSession.h
@@ -0,0 +1,34 @@
+// Copyright Epic Games, Inc. All Rights Reserved.
+
+#pragma once 
+
+#include "Sessions/MLAdapterSession.h"
+#include "UObject/WeakObjectPtrTemplates.h"
+#include "MLAdapterLocalDataCollectionSession.generated.h"
+
+/**
+ * Collects data from Agents' Sensors and writes them to a file for offline processing. Only works with a locally
+ * controlled, single-player game.
+ */
+UCLASS(Blueprintable, EditInlineNew)
+class MLADAPTER_API UMLAdapterLocalDataCollectionSession : public UMLAdapterSession
+{
+	GENERATED_BODY()
+public:
+	virtual void OnPostWorldInit(UWorld& World) override;
+
+	UFUNCTION()
+	virtual void OnPawnControllerChanged(APawn* InPawn, AController* InController);
+
+	virtual void Tick(float DeltaTime) override;
+
+	virtual void Close() override;
+
+	UPROPERTY(EditAnywhere, Category = MLAdapter)
+	FString FileName;
+
+private:
+
+	UPROPERTY()
+	TWeakObjectPtr<UMLAdapterAgent> PlayerControlledAgent;
+};
