MLAdapter: EnhancedInput Sensor/Actuator, support for blueprints and misc. clean-up.  Delete MLAdapterInference

#rnx

diff --git a/Engine/Plugins/AI/MLAdapter/MLAdapter.uplugin b/Engine/Plugins/AI/MLAdapter/MLAdapter.uplugin
--- a/Engine/Plugins/AI/MLAdapter/MLAdapter.uplugin
+++ b/Engine/Plugins/AI/MLAdapter/MLAdapter.uplugin
@@ -30,6 +30,14 @@
 		{
 			"Name": "GameplayAbilities",
 			"Enabled": true
+		},
+		{
+			"Name": "EnhancedInput",
+			"Enabled": true
+		},
+		{
+			"Name": "NeuralNetworkInference",
+			"Enabled": true
 		}
 	]
 }
\ No newline at end of file
diff --git a/Engine/Plugins/AI/MLAdapter/Source/MLAdapter/MLAdapter.Build.cs b/Engine/Plugins/AI/MLAdapter/Source/MLAdapter/MLAdapter.Build.cs
--- a/Engine/Plugins/AI/MLAdapter/Source/MLAdapter/MLAdapter.Build.cs
+++ b/Engine/Plugins/AI/MLAdapter/Source/MLAdapter/MLAdapter.Build.cs
@@ -26,13 +26,15 @@ namespace UnrealBuildTool.Rules
                     "Core",
                     "CoreUObject",
                     "Engine",
+					"EnhancedInput",
                     "GameplayTags",
                     "AIModule",
                     "InputCore",
                     "Json",
                     "JsonUtilities",
                     "GameplayAbilities",
-                }
+					"NeuralNetworkInference"
+				}
             );
 
             PrivateDependencyModuleNames.AddRange(
diff --git a/Engine/Plugins/AI/MLAdapter/Source/MLAdapter/Private/Actuators/MLAdapterActuator_EnhancedInput.cpp b/Engine/Plugins/AI/MLAdapter/Source/MLAdapter/Private/Actuators/MLAdapterActuator_EnhancedInput.cpp
new file mode 100644
--- /dev/null
+++ b/Engine/Plugins/AI/MLAdapter/Source/MLAdapter/Private/Actuators/MLAdapterActuator_EnhancedInput.cpp
@@ -0,0 +1,102 @@
+// Copyright Epic Games, Inc. All Rights Reserved.
+
+#include "Actuators/MLAdapterActuator_EnhancedInput.h"
+#include "MLAdapterTypes.h"
+#include "MLAdapterSpace.h"
+#include "MLAdapterInputHelper.h"
+#include "Agents/MLAdapterAgent.h"
+#include "GameFramework/PlayerController.h"
+#include "Debug/DebugHelpers.h"
+#include "EnhancedInputSubsystems.h"
+#include "Engine/LocalPlayer.h"
+#include "EnhancedPlayerInput.h"
+#include "MLAdapterInputHelper.h"
+
+
+void UMLAdapterActuator_EnhancedInput::Configure(const TMap<FName, FString>& Params)
+{
+	Super::Configure(Params);
+
+	UpdateSpaceDef();
+}
+
+TSharedPtr<FMLAdapter::FSpace> UMLAdapterActuator_EnhancedInput::ConstructSpaceDef() const
+{
+	return FMLAdapterInputHelper::ConstructEnhancedInputSpaceDef(TrackedActions);
+}
+
+void UMLAdapterActuator_EnhancedInput::Act(const float DeltaTime)
+{
+	APlayerController* PC = Cast<APlayerController>(GetControllerAvatar());
+	if (PC == nullptr)
+	{
+		return;
+	}
+
+	FScopeLock Lock(&ActionCS);
+
+	checkf(InputData.Num() % SpaceDef->Num() == 0, TEXT("InputData contains an unexpected number of elements"));
+
+	if (UEnhancedInputLocalPlayerSubsystem* Subsystem = ULocalPlayer::GetSubsystem<UEnhancedInputLocalPlayerSubsystem>(PC->GetLocalPlayer()))
+	{
+		UEnhancedPlayerInput* PlayerInput = Subsystem->GetPlayerInput();
+
+		int32 InputDataIndex = 0;
+		while (InputDataIndex != SpaceDef->Num()) // we might have to loop over all the actions multiple times if digesting is running faster than us
+		{
+			for (const UInputAction* InputAction : TrackedActions)
+			{
+				switch (InputAction->ValueType)
+				{
+				case EInputActionValueType::Boolean:
+				{
+					bool Value = InputData[InputDataIndex] == 1.0f ? true : false;
+					FInputActionValue ActionValue(Value);
+					PlayerInput->InjectInputForAction(InputAction, ActionValue);
+					InputDataIndex += 1;
+					break;
+				}
+				case EInputActionValueType::Axis1D:
+				{
+					float Value = InputData[InputDataIndex];
+					FInputActionValue ActionValue(Value);
+					PlayerInput->InjectInputForAction(InputAction, ActionValue);
+					InputDataIndex += 1;
+					break;
+				}
+				case EInputActionValueType::Axis2D:
+				{
+					FVector2D Value(InputData[InputDataIndex], InputData[InputDataIndex + 1]);
+					FInputActionValue ActionValue(Value);
+					PlayerInput->InjectInputForAction(InputAction, ActionValue);
+					InputDataIndex += 2;
+					break;
+				}
+				case EInputActionValueType::Axis3D:
+				{
+					FVector Value(InputData[InputDataIndex], InputData[InputDataIndex + 1], InputData[InputDataIndex + 2]);
+					FInputActionValue ActionValue(Value);
+					PlayerInput->InjectInputForAction(InputAction, ActionValue);
+					InputDataIndex += 3;
+					break;
+				}
+				default:
+					checkf(false, TEXT("Unsupported value type for input action value!"));
+					break;
+				}
+			}
+		}
+	}
+	
+	InputData.Empty(InputData.Num());
+}
+
+void UMLAdapterActuator_EnhancedInput::DigestInputData(FMLAdapterMemoryReader& ValueStream)
+{
+	FScopeLock Lock(&ActionCS);
+
+	const int32 OldSize = InputData.Num();
+	InputData.AddUninitialized(SpaceDef->Num());
+	// offsetting the serialization since there might be unprocessed data in InputData
+	ValueStream.Serialize(InputData.GetData() + OldSize, SpaceDef->Num() * sizeof(float));
+}
\ No newline at end of file
diff --git a/Engine/Plugins/AI/MLAdapter/Source/MLAdapter/Private/Agents/MLAdapterAgent_Inference.cpp b/Engine/Plugins/AI/MLAdapter/Source/MLAdapter/Private/Agents/MLAdapterAgent_Inference.cpp
new file mode 100644
--- /dev/null
+++ b/Engine/Plugins/AI/MLAdapter/Source/MLAdapter/Private/Agents/MLAdapterAgent_Inference.cpp
@@ -0,0 +1,40 @@
+// Copyright Epic Games, Inc. All Rights Reserved.
+
+#include "Agents/MLAdapterAgent_Inference.h"
+#include "MLAdapterTypes.h"
+#include "HAL/UnrealMemory.h"
+
+void UMLAdapterAgent_Inference::PostInitProperties()
+{
+	UNeuralNetwork* NeuralNetwork = (UNeuralNetwork*)NeuralNetworkPath.TryLoad();
+	
+	if (NeuralNetwork != nullptr)
+	{
+		Brain = NeuralNetwork;
+	}
+
+	Super::PostInitProperties();
+}
+
+void UMLAdapterAgent_Inference::Think(const float DeltaTime)
+{
+	if (Brain == nullptr)
+	{
+		UE_LOG(LogUnrealEditorMLAdapter, Warning, TEXT("Agent beginning to Think but Brain is null, Skipping."));
+		return;
+	}
+
+	TArray<uint8> Buffer;
+	FMLAdapterMemoryWriter Writer(Buffer);
+	GetObservations(Writer);
+
+	const float* DataPtr = (float*)Buffer.GetData();
+	FMemory::Memcpy(Brain->GetInputDataPointerMutable(), DataPtr, Buffer.Num());
+
+	Brain->Run();
+
+	const FNeuralTensor& Tensor = Brain->GetOutputTensor();
+	const TArray<uint8>& Data = Tensor.GetUnderlyingUInt8ArrayRef();
+	FMLAdapterMemoryReader Reader(Data);
+	DigestActions(Reader);
+}
diff --git a/Engine/Plugins/AI/MLAdapter/Source/MLAdapter/Private/Agents/MLAdapterAgents.cpp b/Engine/Plugins/AI/MLAdapter/Source/MLAdapter/Private/Agents/MLAdapterAgents.cpp
--- a/Engine/Plugins/AI/MLAdapter/Source/MLAdapter/Private/Agents/MLAdapterAgents.cpp
+++ b/Engine/Plugins/AI/MLAdapter/Source/MLAdapter/Private/Agents/MLAdapterAgents.cpp
@@ -70,14 +70,22 @@ UMLAdapterAgent::UMLAdapterAgent(const FObjectInitializer& ObjectInitializer)
 	AgentID = FMLAdapter::InvalidAgentID;
 	bEverHadAvatar = false;
 	bRegisteredForPawnControllerChange = false;
+}
+
+void UMLAdapterAgent::PostInitProperties()
+{
+	Super::PostInitProperties();
+
+	if (!HasAnyFlags(RF_ClassDefaultObject))
+	{
+		Actuators.Sort(FAgentElementSort());
+		Sensors.Sort(FAgentElementSort());
 
-	AgentConfig.AddSensor(TEXT("Camera"));
-	AgentConfig.AddSensor(TEXT("Movement"));
-	AgentConfig.AddSensor(TEXT("Input"));
-	AgentConfig.AddActuator(UMLAdapterActuator_InputKey::StaticClass()->GetFName());
-	AgentConfig.Actuators.Add(UMLAdapterActuator_InputKey::StaticClass()->GetFName(), FMLAdapterParameterMap());
-	AgentConfig.AvatarClass = APlayerController::StaticClass();
-	AgentConfig.bAutoRequestNewAvatarUponClearingPrev = true;
+		if (AvatarClass != nullptr)
+		{
+			GetSession().RequestAvatarForAgent(*this);
+		}
+	}
 }
 
 void UMLAdapterAgent::BeginDestroy()
@@ -224,7 +232,7 @@ void UMLAdapterAgent::Configure(const FMLAdapterAgentConfig& NewConfig)
 {
 	ShutDownSensorsAndActuators();
 
-	TSubclassOf<AActor> PreviousAvatarClass = AgentConfig.AvatarClass;
+	TSubclassOf<AActor> PreviousAvatarClass = AvatarClass;
 	AgentConfig = NewConfig;
 
 	for (const TTuple<FName, FMLAdapterParameterMap>& KeyValue : AgentConfig.Actuators)
@@ -257,22 +265,23 @@ void UMLAdapterAgent::Configure(const FMLAdapterAgentConfig& NewConfig)
 
 	if (NewConfig.AvatarClassName != NAME_None)
 	{
-		AgentConfig.AvatarClass = FindObject<UClass>(ANY_PACKAGE, *NewConfig.AvatarClassName.ToString());
+		AvatarClass = FindObject<UClass>(ANY_PACKAGE, *NewConfig.AvatarClassName.ToString());
 	}
-	if (!AgentConfig.AvatarClass)
+
+	if (!AvatarClass)
 	{
-		AgentConfig.AvatarClass = AActor::StaticClass();
+		AvatarClass = AActor::StaticClass();
 	}
 
-	ensure(AgentConfig.AvatarClass || Avatar);
+	ensure(AvatarClass || Avatar);
 
-	if (AgentConfig.AvatarClass && (Avatar == nullptr || IsSuitableAvatar(*Avatar) == false))
+	if (AvatarClass && (Avatar == nullptr || !IsSuitableAvatar(*Avatar)))
 	{
 		SetAvatar(nullptr);
 
 		// if avatar class changed make sure the following RequestAvatarForAgent actually tries to find an avatar
 		// rather than just ignoring the request due to the agent already being in AwaitingAvatar
-		const bool bForceUpdate = (AgentConfig.AvatarClass != PreviousAvatarClass);
+		const bool bForceUpdate = (AvatarClass != PreviousAvatarClass);
 
 		// note that after this call we might not have the avatar just yet 
 		// since the world might not have any. The Session will make sure to 
@@ -333,8 +342,8 @@ void UMLAdapterAgent::GetObservationSpaceDescription(FMLAdapterSpaceDescription&
 bool UMLAdapterAgent::IsSuitableAvatar(AActor& InAvatar) const
 {	
 	return AgentConfig.bAvatarClassExact 
-		? InAvatar.GetClass() == AgentConfig.AvatarClass.Get()
-		: InAvatar.IsA(AgentConfig.AvatarClass);
+		? InAvatar.GetClass() == AvatarClass.Get()
+		: InAvatar.IsA(AvatarClass);
 }
 
 void UMLAdapterAgent::SetAvatar(AActor* InAvatar)
@@ -347,8 +356,8 @@ void UMLAdapterAgent::SetAvatar(AActor* InAvatar)
 
 	if (InAvatar != nullptr && IsSuitableAvatar(*InAvatar) == false)
 	{
-		UE_LOG(LogUnrealEditorMLAdapter, Log, TEXT("SetAvatar was called for agent %u but %s is not a valid avatar (required avatar class %s)")
-			, AgentID, *InAvatar->GetName(), *GetNameSafe(AgentConfig.AvatarClass));
+		UE_LOG(LogUnrealEditorMLAdapter, Log, TEXT("SetAvatar was called for agent %u but %s is not a valid avatar (required avatar class %s)"),
+			AgentID, *InAvatar->GetName(), *GetNameSafe(AvatarClass));
 		return;
 	}
 
diff --git a/Engine/Plugins/AI/MLAdapter/Source/MLAdapter/Private/MLAdapterInputHelper.cpp b/Engine/Plugins/AI/MLAdapter/Source/MLAdapter/Private/MLAdapterInputHelper.cpp
--- a/Engine/Plugins/AI/MLAdapter/Source/MLAdapter/Private/MLAdapterInputHelper.cpp
+++ b/Engine/Plugins/AI/MLAdapter/Source/MLAdapter/Private/MLAdapterInputHelper.cpp
@@ -44,4 +44,55 @@ namespace FMLAdapterInputHelper
 			}
 		}
 	}
+
+	TSharedPtr<FMLAdapter::FSpace> ConstructEnhancedInputSpaceDef(const TArray<UInputAction*>& TrackedActions)
+	{
+		uint32 DiscreteActionCount = 0;
+		uint32 AxisActionCount = 0;
+
+		for (const UInputAction* InputAction : TrackedActions)
+		{
+			switch (InputAction->ValueType)
+			{
+			case EInputActionValueType::Boolean:
+				++DiscreteActionCount;
+				break;
+			case EInputActionValueType::Axis1D:
+				AxisActionCount += 1;
+				break;
+			case EInputActionValueType::Axis2D:
+				AxisActionCount += 2;
+				break;
+			case EInputActionValueType::Axis3D:
+				AxisActionCount += 3;
+				break;
+			default:
+				checkf(false, TEXT("Unsupported value type for input action value!"));
+				break;
+			}
+		}
+
+		FMLAdapter::FSpace* Result = nullptr;
+
+		if (DiscreteActionCount > 0 && AxisActionCount > 0)
+		{
+			Result = new FMLAdapter::FSpace_Tuple({
+				MakeShareable(new FMLAdapter::FSpace_MultiDiscrete(DiscreteActionCount)),
+				MakeShareable(new FMLAdapter::FSpace_Box({AxisActionCount})) });
+		}
+		else if (DiscreteActionCount > 0)
+		{
+			Result = new FMLAdapter::FSpace_MultiDiscrete(DiscreteActionCount);
+		}
+		else if (AxisActionCount > 0)
+		{
+			Result = new FMLAdapter::FSpace_Box({ AxisActionCount });
+		}
+		else
+		{
+			Result = new FMLAdapter::FSpace_Dummy();
+		}
+
+		return MakeShareable(Result);
+	}
 }
\ No newline at end of file
diff --git a/Engine/Plugins/AI/MLAdapter/Source/MLAdapter/Private/MLAdapterSettings.cpp b/Engine/Plugins/AI/MLAdapter/Source/MLAdapter/Private/MLAdapterSettings.cpp
--- a/Engine/Plugins/AI/MLAdapter/Source/MLAdapter/Private/MLAdapterSettings.cpp
+++ b/Engine/Plugins/AI/MLAdapter/Source/MLAdapter/Private/MLAdapterSettings.cpp
@@ -15,21 +15,21 @@ UMLAdapterSettings::UMLAdapterSettings(const FObjectInitializer& ObjectInitializ
 TSubclassOf<UMLAdapterManager> UMLAdapterSettings::GetManagerClass()
 {
 	const FSoftClassPath LocalClassName = GET_CONFIG_VALUE(ManagerClass);
-	TSubclassOf<UMLAdapterManager> LocalClass = LocalClassName.ResolveClass();
+	TSubclassOf<UMLAdapterManager> LocalClass = LocalClassName.TryLoadClass<UMLAdapterManager>();
 	return LocalClass;
 }
 
 TSubclassOf<UMLAdapterSession> UMLAdapterSettings::GetSessionClass()
 {
 	const FSoftClassPath LocalClassName = GET_CONFIG_VALUE(SessionClass);
-	TSubclassOf<UMLAdapterSession> LocalClass = LocalClassName.ResolveClass();
+	TSubclassOf<UMLAdapterSession> LocalClass = LocalClassName.TryLoadClass<UMLAdapterSession>();
 	return LocalClass;
 }
 
 TSubclassOf<UMLAdapterAgent> UMLAdapterSettings::GetAgentClass()
 {
 	const FSoftClassPath LocalClassName = GET_CONFIG_VALUE(DefaultAgentClass);
-	TSubclassOf<UMLAdapterAgent> LocalClass = LocalClassName.ResolveClass();
+	TSubclassOf<UMLAdapterAgent> LocalClass = LocalClassName.TryLoadClass<UMLAdapterAgent>();
 	return LocalClass;
 }
 
diff --git a/Engine/Plugins/AI/MLAdapter/Source/MLAdapter/Private/Sensors/MLAdapterSensor_EnhancedInput.cpp b/Engine/Plugins/AI/MLAdapter/Source/MLAdapter/Private/Sensors/MLAdapterSensor_EnhancedInput.cpp
new file mode 100644
--- /dev/null
+++ b/Engine/Plugins/AI/MLAdapter/Source/MLAdapter/Private/Sensors/MLAdapterSensor_EnhancedInput.cpp
@@ -0,0 +1,171 @@
+// Copyright Epic Games, Inc. All Rights Reserved.
+
+#include "Sensors/MLAdapterSensor_EnhancedInput.h"
+#include "Engine/LocalPlayer.h"
+#include "EnhancedInputSubsystems.h"
+#include "EnhancedInputComponent.h"
+#include "GameFramework/PlayerController.h"
+#include "GameFramework/PlayerInput.h"
+#include "GameFramework/InputSettings.h"
+#include "MLAdapterSpace.h"
+#include "Engine/InputDelegateBinding.h"
+#include "MLAdapterInputHelper.h"
+
+
+void UMLAdapterSensor_EnhancedInput::Configure(const TMap<FName, FString>& Params)
+{
+	Super::Configure(Params);
+
+	UpdateSpaceDef();
+}
+
+void UMLAdapterSensor_EnhancedInput::GetObservations(FMLAdapterMemoryWriter& Ar)
+{
+	FScopeLock Lock(&ObservationCS);
+
+	FMLAdapter::FSpaceSerializeGuard SerializeGuard(SpaceDef, Ar);
+	Ar.Serialize(InputState.GetData(), InputState.Num() * sizeof(float));
+
+	InputState.Empty();
+	InputState.SetNumZeroed(SpaceDef->Num());
+}
+
+TSharedPtr<FMLAdapter::FSpace> UMLAdapterSensor_EnhancedInput::ConstructSpaceDef() const
+{
+	return FMLAdapterInputHelper::ConstructEnhancedInputSpaceDef(TrackedActions);
+}
+
+void UMLAdapterSensor_EnhancedInput::UpdateSpaceDef()
+{
+	Super::UpdateSpaceDef();
+
+	int32 ActionIndex = 0;
+	for (const UInputAction* InputAction : TrackedActions)
+	{
+		InputStateIndices.Add(InputAction->GetName(), ActionIndex);
+
+		switch (InputAction->ValueType)
+		{
+		case EInputActionValueType::Boolean:
+		case EInputActionValueType::Axis1D:
+			ActionIndex += 1;
+			break;
+		case EInputActionValueType::Axis2D:
+			ActionIndex += 2;
+			break;
+		case EInputActionValueType::Axis3D:
+			ActionIndex += 3;
+			break;
+		default:
+			checkf(false, TEXT("Unsupported value type for input action value!"));
+			break;
+		}
+	}
+
+	InputState.SetNumZeroed(SpaceDef->Num());
+}
+
+void UMLAdapterSensor_EnhancedInput::OnAvatarSet(AActor* Avatar)
+{
+	if (Avatar == nullptr)
+	{
+		if (InputComponent != nullptr)
+		{
+			InputComponent->DestroyComponent();
+		}
+
+		return;
+	}
+
+	if (InputComponent != nullptr)
+	{
+		UE_LOG(LogUnrealEditorMLAdapter, Log, TEXT("MLAdapterSensor_EnhancedInput: This sensor already bound to another InputComponent. Cleaning up previous InputComponent and rebinding to new avatar."));
+		InputComponent->DestroyComponent();
+	}
+
+	APawn* Pawn = Cast<APawn>(Avatar);
+	if (Pawn == nullptr)
+	{
+		APlayerController* PC = Cast<APlayerController>(Avatar);
+		if (PC != nullptr)
+		{
+			Pawn = PC->GetPawn();
+		}
+	}
+
+	if (Pawn != nullptr)
+	{
+		UClass* PawnClass = Pawn->GetClass();
+		if (UInputDelegateBinding::SupportsInputDelegate(PawnClass))
+		{
+			InputComponent = NewObject<UEnhancedInputComponent>(Pawn, UInputSettings::GetDefaultInputComponentClass());
+			InputComponent->RegisterComponent();
+
+			for (const UInputAction* InputAction : TrackedActions)
+			{
+				FEnhancedInputActionEventBinding& Binding = InputComponent->BindAction(InputAction, ETriggerEvent::Triggered, this, &UMLAdapterSensor_EnhancedInput::OnInputAction);
+			}
+
+			UInputDelegateBinding::BindInputDelegates(Pawn->GetClass(), InputComponent);
+			UE_LOG(LogUnrealEditorMLAdapter, Log, TEXT("MLAdapterSensor_EnhancedInput: Successfully bound to %s"), *GetNameSafe(PawnClass));
+		}
+	}
+	else
+	{
+		FString AvatarName = Avatar->GetName();
+		UE_LOG(LogUnrealEditorMLAdapter, Warning, TEXT("MLAdapterSensor_EnhancedInput: Unable to bind - could not cast to APawn or APlayerController with a Pawn. Avatar name was %s"), *AvatarName);
+	}
+
+	Super::OnAvatarSet(Avatar);
+}
+
+void UMLAdapterSensor_EnhancedInput::OnInputAction(const FInputActionInstance& ActionInstance)
+{
+	const UInputAction* SourceAction = ActionInstance.GetSourceAction();
+
+	int32* ActionIndexPtr = InputStateIndices.Find(SourceAction->GetName());
+	if (ActionIndexPtr != nullptr)
+	{
+		int32 ActionIndex = *ActionIndexPtr;
+
+		FInputActionValue ActionValue = ActionInstance.GetValue();
+		switch (ActionValue.GetValueType())
+		{
+		case EInputActionValueType::Boolean:
+		{
+			bool Value = ActionValue.Get<bool>();
+			InputState[ActionIndex] = Value ? 1.0f : 0.0f;
+			break;
+		}
+		case EInputActionValueType::Axis1D:
+		{
+			float Value = ActionValue.Get<float>();
+			InputState[ActionIndex] = Value;
+			break;
+		}
+		case EInputActionValueType::Axis2D:
+		{
+			FVector2D Value = ActionValue.Get<FVector2D>();
+			InputState[ActionIndex] = Value.X;
+			InputState[ActionIndex + 1] = Value.Y;
+			break;
+		}
+		case EInputActionValueType::Axis3D:
+		{
+			FVector Value = ActionValue.Get<FVector>();
+			InputState[ActionIndex] = Value.X;
+			InputState[ActionIndex + 1] = Value.Y;
+			InputState[ActionIndex + 2] = Value.Z;
+			break;
+		}
+		default:
+			checkf(false, TEXT("Unsupported value type for input action value!"));
+			break;
+		}
+	}
+	else
+	{
+		FString ActionName = SourceAction->GetName();
+		checkf(false, TEXT("OnInputAction received an ActionInstance from %s that we couldn't find in TrackedActions.  Was TrackedActions modified after binding?"), *ActionName);
+	}
+}
diff --git a/Engine/Plugins/AI/MLAdapter/Source/MLAdapter/Private/Sessions/MLAdapterSession.cpp b/Engine/Plugins/AI/MLAdapter/Source/MLAdapter/Private/Sessions/MLAdapterSession.cpp
--- a/Engine/Plugins/AI/MLAdapter/Source/MLAdapter/Private/Sessions/MLAdapterSession.cpp
+++ b/Engine/Plugins/AI/MLAdapter/Source/MLAdapter/Private/Sessions/MLAdapterSession.cpp
@@ -318,6 +318,7 @@ bool UMLAdapterSession::IsReady() const
 		&& (CachedGameMode->HasMatchEnded() == false);
 }
 
+// todo bmulcahy we probably need to support AddAgent(UClass* Agent) for games that have multiple agents of various types
 FMLAdapter::FAgentID UMLAdapterSession::AddAgent()
 {
 	FScopeLock Lock(&AgentOpCS);
@@ -331,7 +332,6 @@ FMLAdapter::FAgentID UMLAdapterSession::AddAgent()
 	UMLAdapterAgent* NewAgent = FMLAdapter::NewObject<UMLAdapterAgent>(this, AgentClass);
 
 	NewAgent->SetAgentID(Agents.Add(NewAgent));
-	NewAgent->Configure(NewAgent->GetConfig());
 
 	return NewAgent->GetAgentID();
 }
@@ -469,8 +469,8 @@ bool UMLAdapterSession::RequestAvatarForAgent(UMLAdapterAgent& Agent, UWorld* In
 	if (Agent.GetAvatar() != nullptr)
 	{
 		// skipping.
-		UE_LOG(LogUnrealEditorMLAdapter, Verbose, TEXT("UMLAdapterSession::RequestAvatarForAgent called for agent [%s] while it still has an avatar [%s]. Call ClearAvatar first to null-out agent\'s avatar.")
-			, Agent.GetAgentID(), *GetNameSafe(Agent.GetAvatar()));
+		UE_LOG(LogUnrealEditorMLAdapter, Verbose, TEXT("UMLAdapterSession::RequestAvatarForAgent called for agent [%s] while it still has an avatar [%s]. Call ClearAvatar first to null-out agent\'s avatar."),
+			Agent.GetAgentID(), *GetNameSafe(Agent.GetAvatar()));
 		return false;
 	}
 	if (bForceSearch == false && AwaitingAvatar.Find(&Agent) != INDEX_NONE)
@@ -495,7 +495,7 @@ bool UMLAdapterSession::RequestAvatarForAgent(UMLAdapterAgent& Agent, UWorld* In
 	if (InWorld != nullptr)
 	{
 		// @todo might want to make special cases for Controllers and 
-		UClass* AvatarClass = Agent.GetConfig().AvatarClass;
+		UClass* AvatarClass = Agent.AvatarClass;
 		if (ensure(AvatarClass))
 		{
 			for (TActorIterator<AActor> It(InWorld, AvatarClass); It; ++It)
diff --git a/Engine/Plugins/AI/MLAdapter/Source/MLAdapter/Public/Actuators/MLAdapterActuator.h b/Engine/Plugins/AI/MLAdapter/Source/MLAdapter/Public/Actuators/MLAdapterActuator.h
--- a/Engine/Plugins/AI/MLAdapter/Source/MLAdapter/Public/Actuators/MLAdapterActuator.h
+++ b/Engine/Plugins/AI/MLAdapter/Source/MLAdapter/Public/Actuators/MLAdapterActuator.h
@@ -10,7 +10,7 @@
 class AActor;
 struct FMLAdapterDescription;
 
-UCLASS(Abstract)
+UCLASS(Abstract, Blueprintable, EditInlineNew)
 class MLADAPTER_API UMLAdapterActuator : public UMLAdapterAgentElement
 {
 	GENERATED_BODY()
diff --git a/Engine/Plugins/AI/MLAdapter/Source/MLAdapter/Public/Actuators/MLAdapterActuator_EnhancedInput.h b/Engine/Plugins/AI/MLAdapter/Source/MLAdapter/Public/Actuators/MLAdapterActuator_EnhancedInput.h
new file mode 100644
--- /dev/null
+++ b/Engine/Plugins/AI/MLAdapter/Source/MLAdapter/Public/Actuators/MLAdapterActuator_EnhancedInput.h
@@ -0,0 +1,30 @@
+// Copyright Epic Games, Inc. All Rights Reserved.
+
+#pragma once
+
+#include "UObject/Object.h"
+#include "InputCoreTypes.h"
+#include "Actuators/MLAdapterActuator.h"
+#include "InputAction.h"
+#include "MLAdapterActuator_EnhancedInput.generated.h"
+
+
+UCLASS()
+class MLADAPTER_API UMLAdapterActuator_EnhancedInput : public UMLAdapterActuator
+{
+	GENERATED_BODY()
+
+public:
+	virtual void Configure(const TMap<FName, FString>& Params) override;
+	virtual TSharedPtr<FMLAdapter::FSpace> ConstructSpaceDef() const override;
+
+	virtual void Act(const float DeltaTime) override;
+
+	virtual void DigestInputData(FMLAdapterMemoryReader& ValueStream) override;
+	
+	UPROPERTY(EditDefaultsOnly, BlueprintReadOnly, Category = MLAdapter)
+	TArray<TObjectPtr<UInputAction>> TrackedActions;
+
+protected:
+	TArray<float> InputData;
+};
\ No newline at end of file
diff --git a/Engine/Plugins/AI/MLAdapter/Source/MLAdapter/Public/Agents/MLAdapterAgent.h b/Engine/Plugins/AI/MLAdapter/Source/MLAdapter/Public/Agents/MLAdapterAgent.h
--- a/Engine/Plugins/AI/MLAdapter/Source/MLAdapter/Public/Agents/MLAdapterAgent.h
+++ b/Engine/Plugins/AI/MLAdapter/Source/MLAdapter/Public/Agents/MLAdapterAgent.h
@@ -19,7 +19,7 @@ struct FMLAdapterSpaceDescription;
 
 
 USTRUCT()
-struct FMLAdapterParameterMap
+struct MLADAPTER_API FMLAdapterParameterMap
 {
 	GENERATED_BODY()
 
@@ -28,7 +28,7 @@ struct FMLAdapterParameterMap
 };
 
 USTRUCT()
-struct FMLAdapterAgentConfig
+struct MLADAPTER_API FMLAdapterAgentConfig
 {
 	GENERATED_BODY()
 
@@ -51,8 +51,6 @@ struct FMLAdapterAgentConfig
 	UPROPERTY()
 	bool bAutoRequestNewAvatarUponClearingPrev = true;
 
-	TSubclassOf<AActor> AvatarClass;
-
 	FMLAdapterParameterMap& AddSensor(const FName SensorName, FMLAdapterParameterMap&& Parameters = FMLAdapterParameterMap());
 	FMLAdapterParameterMap& AddActuator(const FName ActuatorName, FMLAdapterParameterMap&& Parameters = FMLAdapterParameterMap());
 };
@@ -69,6 +67,9 @@ class MLADAPTER_API UMLAdapterAgent : public UObject
     GENERATED_BODY()
 public:
 	UMLAdapterAgent(const FObjectInitializer& ObjectInitializer = FObjectInitializer::Get());
+
+	virtual void PostInitProperties() override;
+
 	virtual void BeginDestroy() override;
 
 	virtual bool RegisterSensor(UMLAdapterSensor& Sensor);
@@ -138,14 +139,16 @@ public:
 protected:
 	virtual void ShutDownSensorsAndActuators();
 
-
 protected:
-	UPROPERTY(EditAnywhere, BlueprintReadWrite, Category = MLAdapter)
+	UPROPERTY(EditDefaultsOnly, Instanced, BlueprintReadOnly, Category = MLAdapter)
 	TArray<UMLAdapterSensor*> Sensors;
 
-	UPROPERTY(EditAnywhere, BlueprintReadWrite, Category = MLAdapter)
+	UPROPERTY(EditDefaultsOnly, Instanced, BlueprintReadOnly, Category = MLAdapter)
 	TArray<UMLAdapterActuator*> Actuators;
 
+	UPROPERTY(EditDefaultsOnly, BlueprintReadOnly, Category = MLAdapter)
+	TSubclassOf<AActor> AvatarClass;
+
 private:
 	UPROPERTY()
 	AActor* Avatar;
diff --git a/Engine/Plugins/AI/MLAdapter/Source/MLAdapter/Public/Agents/MLAdapterAgent_Inference.h b/Engine/Plugins/AI/MLAdapter/Source/MLAdapter/Public/Agents/MLAdapterAgent_Inference.h
new file mode 100644
--- /dev/null
+++ b/Engine/Plugins/AI/MLAdapter/Source/MLAdapter/Public/Agents/MLAdapterAgent_Inference.h
@@ -0,0 +1,28 @@
+// Copyright Epic Games, Inc. All Rights Reserved.
+
+#pragma once
+
+#include "Agents/MLAdapterAgent.h"
+#include "NeuralNetwork.h"
+#include "MLAdapterAgent_Inference.generated.h"
+
+/**
+ * Inference agents have a neural network that can process senses and output actuations via their Think method. You 
+ * can create a blueprint of this class to easily wire-up an agent that functions entirely inside the Unreal Engine.
+ * Typically, this class would be used in conjunction with MLAdapterNoRPCManager.
+ */
+UCLASS(Blueprintable, EditInlineNew)
+class MLADAPTER_API UMLAdapterAgent_Inference : public UMLAdapterAgent
+{
+	GENERATED_BODY()
+public:
+	virtual void PostInitProperties() override;
+
+	virtual void Think(const float DeltaTime) override;
+
+	UPROPERTY(EditDefaultsOnly, Category = MLAdapter, meta = (AllowedClasses = "NeuralNetwork"))
+	FSoftObjectPath NeuralNetworkPath;
+
+	UPROPERTY()
+	UNeuralNetwork* Brain;
+};
diff --git a/Engine/Plugins/AI/MLAdapter/Source/MLAdapter/Public/MLAdapterInputHelper.h b/Engine/Plugins/AI/MLAdapter/Source/MLAdapter/Public/MLAdapterInputHelper.h
--- a/Engine/Plugins/AI/MLAdapter/Source/MLAdapter/Public/MLAdapterInputHelper.h
+++ b/Engine/Plugins/AI/MLAdapter/Source/MLAdapter/Public/MLAdapterInputHelper.h
@@ -3,8 +3,12 @@
 #pragma once
 
 #include "InputCoreTypes.h"
+#include "MLAdapterSpace.h"
+#include "EnhancedPlayerInput.h"
 
 namespace FMLAdapterInputHelper
 {
 	MLADAPTER_API void CreateInputMap(TArray<TTuple<FKey, FName>>& InterfaceKeys, TMap<FKey, int32>& FKeyToInterfaceKeyMap);
+
+	MLADAPTER_API TSharedPtr<FMLAdapter::FSpace> ConstructEnhancedInputSpaceDef(const TArray<UInputAction*>& TrackedActions);
 }
\ No newline at end of file
diff --git a/Engine/Plugins/AI/MLAdapter/Source/MLAdapter/Public/Sensors/MLAdapterSensor_EnhancedInput.h b/Engine/Plugins/AI/MLAdapter/Source/MLAdapter/Public/Sensors/MLAdapterSensor_EnhancedInput.h
new file mode 100644
--- /dev/null
+++ b/Engine/Plugins/AI/MLAdapter/Source/MLAdapter/Public/Sensors/MLAdapterSensor_EnhancedInput.h
@@ -0,0 +1,38 @@
+// Copyright Epic Games, Inc. All Rights Reserved.
+
+#pragma once
+
+#include "Sensors/MLAdapterSensor.h"
+#include "InputAction.h"
+#include "EnhancedInputComponent.h"
+#include "MLAdapterSensor_EnhancedInput.generated.h"
+
+
+UCLASS(Blueprintable)
+class MLADAPTER_API UMLAdapterSensor_EnhancedInput : public UMLAdapterSensor
+{
+	GENERATED_BODY()
+
+protected:
+	virtual void Configure(const TMap<FName, FString>& Params) override;
+	virtual void GetObservations(FMLAdapterMemoryWriter& Ar) override;
+	virtual void UpdateSpaceDef() override;
+	virtual TSharedPtr<FMLAdapter::FSpace> ConstructSpaceDef() const override;
+
+	/** The Avatar here is expected to be a player-controlled APawn or an APlayerController */
+	virtual void OnAvatarSet(AActor* Avatar) override;
+
+	TArray<float> InputState;
+
+	TMap<FString, int32> InputStateIndices;
+
+public:
+	UPROPERTY(EditDefaultsOnly, BlueprintReadOnly, Category = MLAdapter)
+	TArray<TObjectPtr<UInputAction>> TrackedActions;
+
+	UPROPERTY(BlueprintReadOnly, Category = MLAdapter)
+	TObjectPtr<UEnhancedInputComponent> InputComponent;
+
+	UFUNCTION()
+	void OnInputAction(const FInputActionInstance& ActionInstance);
+};
diff --git a/Engine/Plugins/AI/MLAdapter/Source/MLAdapterTestSuite/Private/AgentTest.cpp b/Engine/Plugins/AI/MLAdapter/Source/MLAdapterTestSuite/Private/AgentTest.cpp
--- a/Engine/Plugins/AI/MLAdapter/Source/MLAdapterTestSuite/Private/AgentTest.cpp
+++ b/Engine/Plugins/AI/MLAdapter/Source/MLAdapterTestSuite/Private/AgentTest.cpp
@@ -15,7 +15,6 @@ struct FMockAgentConfig_Actors : public FMLAdapterAgentConfig
 {
 	FMockAgentConfig_Actors()
 	{
-		AvatarClass = AActor::StaticClass();
 		bAutoRequestNewAvatarUponClearingPrev = false;
 	}
 };
@@ -24,7 +23,6 @@ struct FMockAgentConfig_OnlyActor : public FMLAdapterAgentConfig
 {
 	FMockAgentConfig_OnlyActor()
 	{
-		AvatarClass = AActor::StaticClass();
 		bAvatarClassExact = true;
 		bAutoRequestNewAvatarUponClearingPrev = false;
 	}
diff --git a/Engine/Plugins/Experimental/MLAdapterInference/MLAdapterInference.uplugin b/Engine/Plugins/Experimental/MLAdapterInference/MLAdapterInference.uplugin
deleted file mode 100644
diff --git a/Engine/Plugins/Experimental/MLAdapterInference/README.md b/Engine/Plugins/Experimental/MLAdapterInference/README.md
deleted file mode 100644
diff --git a/Engine/Plugins/Experimental/MLAdapterInference/Source/MLAdapterInference/MLAdapterInference.Build.cs b/Engine/Plugins/Experimental/MLAdapterInference/Source/MLAdapterInference/MLAdapterInference.Build.cs
deleted file mode 100644
diff --git a/Engine/Plugins/Experimental/MLAdapterInference/Source/MLAdapterInference/Private/MLAdapterInferenceAgent.cpp b/Engine/Plugins/Experimental/MLAdapterInference/Source/MLAdapterInference/Private/MLAdapterInferenceAgent.cpp
deleted file mode 100644
diff --git a/Engine/Plugins/Experimental/MLAdapterInference/Source/MLAdapterInference/Private/MLAdapterInferenceManager.cpp b/Engine/Plugins/Experimental/MLAdapterInference/Source/MLAdapterInference/Private/MLAdapterInferenceManager.cpp
deleted file mode 100644
diff --git a/Engine/Plugins/Experimental/MLAdapterInference/Source/MLAdapterInference/Private/MLAdapterInferenceModule.cpp b/Engine/Plugins/Experimental/MLAdapterInference/Source/MLAdapterInference/Private/MLAdapterInferenceModule.cpp
deleted file mode 100644
diff --git a/Engine/Plugins/Experimental/MLAdapterInference/Source/MLAdapterInference/Private/MLAdapterInferenceSettings.cpp b/Engine/Plugins/Experimental/MLAdapterInference/Source/MLAdapterInference/Private/MLAdapterInferenceSettings.cpp
deleted file mode 100644
diff --git a/Engine/Plugins/Experimental/MLAdapterInference/Source/MLAdapterInference/Private/MLAdapterInferenceTypes.cpp b/Engine/Plugins/Experimental/MLAdapterInference/Source/MLAdapterInference/Private/MLAdapterInferenceTypes.cpp
deleted file mode 100644
diff --git a/Engine/Plugins/Experimental/MLAdapterInference/Source/MLAdapterInference/Public/MLAdapterInferenceAgent.h b/Engine/Plugins/Experimental/MLAdapterInference/Source/MLAdapterInference/Public/MLAdapterInferenceAgent.h
deleted file mode 100644
diff --git a/Engine/Plugins/Experimental/MLAdapterInference/Source/MLAdapterInference/Public/MLAdapterInferenceManager.h b/Engine/Plugins/Experimental/MLAdapterInference/Source/MLAdapterInference/Public/MLAdapterInferenceManager.h
deleted file mode 100644
diff --git a/Engine/Plugins/Experimental/MLAdapterInference/Source/MLAdapterInference/Public/MLAdapterInferenceModule.h b/Engine/Plugins/Experimental/MLAdapterInference/Source/MLAdapterInference/Public/MLAdapterInferenceModule.h
deleted file mode 100644
diff --git a/Engine/Plugins/Experimental/MLAdapterInference/Source/MLAdapterInference/Public/MLAdapterInferenceSettings.h b/Engine/Plugins/Experimental/MLAdapterInference/Source/MLAdapterInference/Public/MLAdapterInferenceSettings.h
deleted file mode 100644
diff --git a/Engine/Plugins/Experimental/MLAdapterInference/Source/MLAdapterInference/Public/MLAdapterInferenceTypes.h b/Engine/Plugins/Experimental/MLAdapterInference/Source/MLAdapterInference/Public/MLAdapterInferenceTypes.h
deleted file mode 100644
