MLAdapter: add missing include Pawn.h to Sensor_EnhancedInput

#rnx

diff --git a/Engine/Plugins/AI/MLAdapter/Source/MLAdapter/Private/Sensors/MLAdapterSensor_EnhancedInput.cpp b/Engine/Plugins/AI/MLAdapter/Source/MLAdapter/Private/Sensors/MLAdapterSensor_EnhancedInput.cpp
--- a/Engine/Plugins/AI/MLAdapter/Source/MLAdapter/Private/Sensors/MLAdapterSensor_EnhancedInput.cpp
+++ b/Engine/Plugins/AI/MLAdapter/Source/MLAdapter/Private/Sensors/MLAdapterSensor_EnhancedInput.cpp
@@ -5,6 +5,7 @@
 #include "EnhancedInputSubsystems.h"
 #include "EnhancedInputComponent.h"
 #include "GameFramework/PlayerController.h"
+#include "GameFramework/Pawn.h"
 #include "GameFramework/PlayerInput.h"
 #include "GameFramework/InputSettings.h"
 #include "MLAdapterSpace.h"
