MLAdapter: target Linux platform instead of Unix

#rnx

diff --git a/Engine/Plugins/AI/MLAdapter/Source/MLAdapter/MLAdapter.Build.cs b/Engine/Plugins/AI/MLAdapter/Source/MLAdapter/MLAdapter.Build.cs
--- a/Engine/Plugins/AI/MLAdapter/Source/MLAdapter/MLAdapter.Build.cs
+++ b/Engine/Plugins/AI/MLAdapter/Source/MLAdapter/MLAdapter.Build.cs
@@ -46,7 +46,7 @@ namespace UnrealBuildTool.Rules
             // RPCLib disabled on other platforms
             if (Target.Platform == UnrealTargetPlatform.Win64 ||
 				Target.Platform == UnrealTargetPlatform.Mac ||
-				Target.IsInPlatformGroup(UnrealPlatformGroup.Unix))
+				Target.Platform == UnrealTargetPlatform.Linux)
             {
                 PublicDefinitions.Add("WITH_RPCLIB=1");
 				PrivateDependencyModuleNames.Add("RPCLib");
diff --git a/Engine/Plugins/AI/MLAdapter/Source/MLAdapterTestSuite/MLAdapterTestSuite.Build.cs b/Engine/Plugins/AI/MLAdapter/Source/MLAdapterTestSuite/MLAdapterTestSuite.Build.cs
--- a/Engine/Plugins/AI/MLAdapter/Source/MLAdapterTestSuite/MLAdapterTestSuite.Build.cs
+++ b/Engine/Plugins/AI/MLAdapter/Source/MLAdapterTestSuite/MLAdapterTestSuite.Build.cs
@@ -36,7 +36,7 @@ namespace UnrealBuildTool.Rules
 			// RPCLib disabled on other platforms
 			if (Target.Platform == UnrealTargetPlatform.Win64 ||
 				Target.Platform == UnrealTargetPlatform.Mac ||
-				Target.IsInPlatformGroup(UnrealPlatformGroup.Unix))
+				Target.Platform == UnrealTargetPlatform.Linux)
 			{
 				PublicDefinitions.Add("WITH_RPCLIB=1");
 				PrivateDependencyModuleNames.Add("RPCLib");
