MLAdapter: Add API Reference comments to most public members & small amounts of clean-up.

#rnx

diff --git a/Engine/Plugins/AI/MLAdapter/MLAdapter.uplugin b/Engine/Plugins/AI/MLAdapter/MLAdapter.uplugin
--- a/Engine/Plugins/AI/MLAdapter/MLAdapter.uplugin
+++ b/Engine/Plugins/AI/MLAdapter/MLAdapter.uplugin
@@ -2,8 +2,8 @@
 	"FileVersion": 3,
 	"Version": 1,
 	"VersionName": "0.0.1",
-	"FriendlyName": "Unreal Engine support for Machine Learning",
-	"Description": "[EXPERIMENTAL] Exposes means of having external agents query world and game state and control in-game agents",
+	"FriendlyName": "ML Adapter",
+	"Description": "[EXPERIMENTAL] A framework for training and utilizing machine learning agents in games. Creates an RPC interface through which an external process can query game state and control in-game actors. Once trained, agents can be run in-engine via neural networks loaded from ONNX models.",
 	"Category": "AI",
 	"CreatedBy": "Epic Games, Inc.",
 	"CreatedByURL": "http://epicgames.com",
diff --git a/Engine/Plugins/AI/MLAdapter/Source/MLAdapter/Private/Actuators/MLAdapterActuator_Camera.cpp b/Engine/Plugins/AI/MLAdapter/Source/MLAdapter/Private/Actuators/MLAdapterActuator_Camera.cpp
--- a/Engine/Plugins/AI/MLAdapter/Source/MLAdapter/Private/Actuators/MLAdapterActuator_Camera.cpp
+++ b/Engine/Plugins/AI/MLAdapter/Source/MLAdapter/Private/Actuators/MLAdapterActuator_Camera.cpp
@@ -23,7 +23,6 @@ UMLAdapterActuator_Camera::UMLAdapterActuator_Camera(const FObjectInitializer& O
 {
 	bVectorMode = true;
 	bConsumeData = true;
-	bDeltaMode = true;
 }
 
 void UMLAdapterActuator_Camera::Configure(const TMap<FName, FString>& Params)
diff --git a/Engine/Plugins/AI/MLAdapter/Source/MLAdapter/Private/Sensors/MLAdapterSensor.cpp b/Engine/Plugins/AI/MLAdapter/Source/MLAdapter/Private/Sensors/MLAdapterSensor.cpp
--- a/Engine/Plugins/AI/MLAdapter/Source/MLAdapter/Private/Sensors/MLAdapterSensor.cpp
+++ b/Engine/Plugins/AI/MLAdapter/Source/MLAdapter/Private/Sensors/MLAdapterSensor.cpp
@@ -26,6 +26,7 @@ UMLAdapterSensor::UMLAdapterSensor(const FObjectInitializer& ObjectInitializer)
 void UMLAdapterSensor::PostInitProperties()
 {
 	Super::PostInitProperties();
+
 	if (HasAnyFlags(RF_ClassDefaultObject))
 	{
 		if (GetClass()->HasAnyClassFlags(CLASS_Abstract) == false)
@@ -36,13 +37,15 @@ void UMLAdapterSensor::PostInitProperties()
 	else
 	{
 		const UMLAdapterSensor* CDO = GetDefault<UMLAdapterSensor>(GetClass());
-		// already checked 
+		check(CDO);
 		ElementID = CDO->ElementID;
 	}
 }
 
 void UMLAdapterSensor::Configure(const TMap<FName, FString>& Params)
 {
+	Super::Configure(Params);
+
 	const FName NAME_EveryTick = TEXT("tick_every_frame");
 	const FName NAME_EveryNTicks = TEXT("tick_every_n_frames");
 	const FName NAME_EveryXFrames = TEXT("tick_every_x_seconds");
@@ -136,5 +139,5 @@ void UMLAdapterSensor::OnPawnChanged(APawn* OldPawn, APawn* NewPawn)
 
 void UMLAdapterSensor::ClearPawn(APawn& InPawn)
 {
-
+	// Can be overridden in derived classes
 }
diff --git a/Engine/Plugins/AI/MLAdapter/Source/MLAdapter/Private/Sensors/MLAdapterSensor_EnhancedInput.cpp b/Engine/Plugins/AI/MLAdapter/Source/MLAdapter/Private/Sensors/MLAdapterSensor_EnhancedInput.cpp
--- a/Engine/Plugins/AI/MLAdapter/Source/MLAdapter/Private/Sensors/MLAdapterSensor_EnhancedInput.cpp
+++ b/Engine/Plugins/AI/MLAdapter/Source/MLAdapter/Private/Sensors/MLAdapterSensor_EnhancedInput.cpp
@@ -127,6 +127,8 @@ void UMLAdapterSensor_EnhancedInput::OnInputAction(const FInputActionInstance& A
 	int32* ActionIndexPtr = InputStateIndices.Find(SourceAction->GetName());
 	if (ActionIndexPtr != nullptr)
 	{
+		FScopeLock Lock(&ObservationCS);
+
 		int32 ActionIndex = *ActionIndexPtr;
 
 		FInputActionValue ActionValue = ActionInstance.GetValue();
diff --git a/Engine/Plugins/AI/MLAdapter/Source/MLAdapter/Public/Actuators/MLAdapterActuator.h b/Engine/Plugins/AI/MLAdapter/Source/MLAdapter/Public/Actuators/MLAdapterActuator.h
--- a/Engine/Plugins/AI/MLAdapter/Source/MLAdapter/Public/Actuators/MLAdapterActuator.h
+++ b/Engine/Plugins/AI/MLAdapter/Source/MLAdapter/Public/Actuators/MLAdapterActuator.h
@@ -2,6 +2,9 @@
 
 #pragma once
 
+#include "GameFramework/Actor.h"
+#include "UObject/UObjectGlobals.h"
+#include "HAL/CriticalSection.h"
 #include "Agents/MLAdapterAgentElement.h"
 #include "MLAdapterTypes.h"
 #include "MLAdapterActuator.generated.h"
@@ -10,17 +13,28 @@
 class AActor;
 struct FMLAdapterDescription;
 
+/** Allows an agent to take an action in the game world. */
 UCLASS(Abstract, Blueprintable, EditInlineNew)
 class MLADAPTER_API UMLAdapterActuator : public UMLAdapterAgentElement
 {
 	GENERATED_BODY()
+
 public:
 	UMLAdapterActuator(const FObjectInitializer& ObjectInitializer = FObjectInitializer::Get());
+
 	virtual void PostInitProperties() override;
 
+	/** Takes an action in the game based on earlier digested data. */
 	virtual void Act(const float DeltaTime) {}
+
+	/** Stores data from the ValueStream to be later used by Act. */
 	virtual void DigestInputData(FMLAdapterMemoryReader& ValueStream) {}
 
 protected:
+
+	/**
+	 * This needs to be used in every implementation of Act and DigestInputData to ensure that shared state is not be
+	 * accessed simultaneously by the game thread and the RPC server thread.
+	 */
 	mutable FCriticalSection ActionCS;
 };
\ No newline at end of file
diff --git a/Engine/Plugins/AI/MLAdapter/Source/MLAdapter/Public/Actuators/MLAdapterActuator_Camera.h b/Engine/Plugins/AI/MLAdapter/Source/MLAdapter/Public/Actuators/MLAdapterActuator_Camera.h
--- a/Engine/Plugins/AI/MLAdapter/Source/MLAdapter/Public/Actuators/MLAdapterActuator_Camera.h
+++ b/Engine/Plugins/AI/MLAdapter/Source/MLAdapter/Public/Actuators/MLAdapterActuator_Camera.h
@@ -8,19 +8,23 @@
 #include "MLAdapterActuator_Camera.generated.h"
 
 
-UCLASS()
+/** Allows an agent to rotate the camera. */
+UCLASS(Blueprintable, EditInlineNew)
 class MLADAPTER_API UMLAdapterActuator_Camera : public UMLAdapterActuator
 {
 	GENERATED_BODY()
+
 public:
 	UMLAdapterActuator_Camera(const FObjectInitializer& ObjectInitializer = FObjectInitializer::Get());
 
 	virtual void Configure(const TMap<FName, FString>& Params) override;
+
 	virtual TSharedPtr<FMLAdapter::FSpace> ConstructSpaceDef() const override;
 
-	/** Presses the keys stored in "KeysToPress" */
+	/** Rotates the camera. */
 	virtual void Act(const float DeltaTime) override;
 
+	/** Stores data in HeadingVector or HeadingRotator depending on bVectorMode. */
 	virtual void DigestInputData(FMLAdapterMemoryReader& ValueStream) override;
 
 protected:
@@ -29,10 +33,19 @@ protected:
 	virtual void DescribeSelfToGameplayDebugger(class FGameplayDebuggerCategory& DebuggerCategory) const override;
 #endif // WITH_GAMEPLAY_DEBUGGER
 
+	/** Stores the direction to point the camera as a rotator. */
+	UPROPERTY()
 	FRotator HeadingRotator;
+
+	/** Stores the direction to point the camera as a vector. */
+	UPROPERTY()
 	FVector HeadingVector;
 
+	/** If true, acting will reinitialize the HeadingRotator/Vector after reading them. */
+	UPROPERTY()
 	uint32 bConsumeData : 1;
+
+	/** If true, use the HeadingVector. Otherwise, use the HeadingRotator. */
+	UPROPERTY()
 	uint32 bVectorMode : 1;
-	uint32 bDeltaMode : 1;
 };
\ No newline at end of file
diff --git a/Engine/Plugins/AI/MLAdapter/Source/MLAdapter/Public/Actuators/MLAdapterActuator_EnhancedInput.h b/Engine/Plugins/AI/MLAdapter/Source/MLAdapter/Public/Actuators/MLAdapterActuator_EnhancedInput.h
--- a/Engine/Plugins/AI/MLAdapter/Source/MLAdapter/Public/Actuators/MLAdapterActuator_EnhancedInput.h
+++ b/Engine/Plugins/AI/MLAdapter/Source/MLAdapter/Public/Actuators/MLAdapterActuator_EnhancedInput.h
@@ -9,22 +9,31 @@
 #include "MLAdapterActuator_EnhancedInput.generated.h"
 
 
-UCLASS()
+/** Allows an agent to simulate player input via the UEnhancedPlayerInput subsystem. */
+UCLASS(Blueprintable, EditInlineNew)
 class MLADAPTER_API UMLAdapterActuator_EnhancedInput : public UMLAdapterActuator
 {
 	GENERATED_BODY()
 
 public:
 	virtual void Configure(const TMap<FName, FString>& Params) override;
+
 	virtual TSharedPtr<FMLAdapter::FSpace> ConstructSpaceDef() const override;
 
+	/** Performs one or more actions on the input component. */
 	virtual void Act(const float DeltaTime) override;
 
+	/** Stores data in InputData to be later injected into actions. */
 	virtual void DigestInputData(FMLAdapterMemoryReader& ValueStream) override;
-	
-	UPROPERTY(EditDefaultsOnly, BlueprintReadOnly, Category = MLAdapter)
-	TArray<TObjectPtr<UInputAction>> TrackedActions;
 
 protected:
 	TArray<float> InputData;
+
+	/**
+	* The actions that this actuator is allowed to inject data into. Modifying this list after using the actuator
+	* is unsupported because the InputData is injected into the actions based on the ordering in this array. Changing
+	* this could cause data to be injected incorrectly.
+	*/
+	UPROPERTY(EditDefaultsOnly, BlueprintReadOnly, Category = MLAdapter)
+	TArray<TObjectPtr<UInputAction>> TrackedActions;
 };
\ No newline at end of file
diff --git a/Engine/Plugins/AI/MLAdapter/Source/MLAdapter/Public/Actuators/MLAdapterActuator_InputKey.h b/Engine/Plugins/AI/MLAdapter/Source/MLAdapter/Public/Actuators/MLAdapterActuator_InputKey.h
--- a/Engine/Plugins/AI/MLAdapter/Source/MLAdapter/Public/Actuators/MLAdapterActuator_InputKey.h
+++ b/Engine/Plugins/AI/MLAdapter/Source/MLAdapter/Public/Actuators/MLAdapterActuator_InputKey.h
@@ -7,15 +7,17 @@
 #include "Actuators/MLAdapterActuator.h"
 #include "MLAdapterActuator_InputKey.generated.h"
 
-
-UCLASS()
+/** Allows an agent to directly inject key presses into its avatar's input component. */
+UCLASS(Blueprintable, EditInlineNew)
 class MLADAPTER_API UMLAdapterActuator_InputKey : public UMLAdapterActuator
 {
 	GENERATED_BODY()
+
 public:
 	UMLAdapterActuator_InputKey(const FObjectInitializer& ObjectInitializer = FObjectInitializer::Get());
 
 	virtual void Configure(const TMap<FName, FString>& Params) override;
+
 	virtual TSharedPtr<FMLAdapter::FSpace> ConstructSpaceDef() const override;
 
 	/** Presses the keys stored in "KeysToPress" */
@@ -25,6 +27,7 @@ public:
 
 protected:
 	TArray<TTuple<FKey, FName>> RegisteredKeys;
+
 	TArray<int32> KeysToPress;
 
 	TBitArray<> PressedKeys;
diff --git a/Engine/Plugins/AI/MLAdapter/Source/MLAdapter/Public/Agents/MLAdapterAgent.h b/Engine/Plugins/AI/MLAdapter/Source/MLAdapter/Public/Agents/MLAdapterAgent.h
--- a/Engine/Plugins/AI/MLAdapter/Source/MLAdapter/Public/Agents/MLAdapterAgent.h
+++ b/Engine/Plugins/AI/MLAdapter/Source/MLAdapter/Public/Agents/MLAdapterAgent.h
@@ -18,6 +18,7 @@ class UMLAdapterActuator;
 struct FMLAdapterSpaceDescription;
 
 
+/** Provides a serializable mapping from parameter name to value that is used to configure sensors & actuators. */
 USTRUCT()
 struct MLADAPTER_API FMLAdapterParameterMap
 {
@@ -27,6 +28,7 @@ struct MLADAPTER_API FMLAdapterParameterMap
 	TMap<FName, FString> Params;
 };
 
+/** A serializable config for an agent. Used in the external API to define agents. */
 USTRUCT()
 struct MLADAPTER_API FMLAdapterAgentConfig
 {
@@ -44,7 +46,7 @@ struct MLADAPTER_API FMLAdapterAgentConfig
 	UPROPERTY()
 	FName AgentClassName;
 
-	/** if set to true won't accept child classes of AvatarClass */
+	/** If set to true, won't accept child classes of the AvatarClass. */
 	UPROPERTY()
 	bool bAvatarClassExact = false;
 
@@ -57,47 +59,71 @@ struct MLADAPTER_API FMLAdapterAgentConfig
 
 namespace FMLAdapterAgentHelpers
 {
+	/** Get the Pawn and Controller from the given Avatar if possible. */
 	bool MLADAPTER_API GetAsPawnAndController(AActor* Avatar, AController*& OutController, APawn*& OutPawn);
 }
 
-
+/**
+ * An agent capable of controlling a single avatar (e.g. a Pawn or Controller). Contains sensors for
+ * perceiving information about the environment and actuators for taking actions in the game.
+ */
 UCLASS(Blueprintable, EditInlineNew)
 class MLADAPTER_API UMLAdapterAgent : public UObject
 {
     GENERATED_BODY()
+
 public:
 	UMLAdapterAgent(const FObjectInitializer& ObjectInitializer = FObjectInitializer::Get());
 
+	/** Perform initial setup for blueprint spawned agents. */
 	virtual void PostInitProperties() override;
 
+	/** Shutdown all the sensors and actuators and cleanup references to the avatar. */
 	virtual void BeginDestroy() override;
 
+	/** Add a sensor to this agent. Returns true if the sensor was successfully added. */
 	virtual bool RegisterSensor(UMLAdapterSensor& Sensor);
 	
-	/** Updates all the senses that are configured as 'IsPolling'*/
+	/** Updates all the sensors that are configured as 'IsPolling'. */
 	virtual void Sense(const float DeltaTime);
 	
-	// Decide what action to take based on the current observations
+	/** 
+	 * Decide what action to take based on the current observations.
+	 * @see UMLAdapterAgent_Inference for an agent that implements this.
+	 */
 	virtual void Think(const float DeltaTime);
 
-	// trigger all of the agent's 
+	/** Tick all of the agent's actuators. */
 	virtual void Act(const float DeltaTime);
 	
-	//virtual void DigestActions(const std::vector<float>& ValueStream);
+	/** Move data into the actuators for the next time Act is called. */
 	virtual void DigestActions(FMLAdapterMemoryReader& ValueStream);
 
+	/** Get this agent's ID. */
 	FMLAdapter::FAgentID GetAgentID() const { return AgentID; }
+
+	/** Get the Pawn this agent is controlling. */
 	APawn* GetPawn() { return Pawn; }
+
+	/** Get the Pawn this agent is controlling. */
 	const APawn* GetPawn() const { return Pawn; }
+
+	/** Get the Controller this agent is controlling. */
 	AController* GetController() { return Controller; }
+
+	/** Get the Controller this agent is controlling. */
 	const AController* GetController() const { return Controller; }
 	
 	TArray<UMLAdapterSensor*>::TConstIterator GetSensorsConstIterator() const { return Sensors.CreateConstIterator(); }
 	TArray<UMLAdapterActuator*>::TConstIterator GetActuatorsConstIterator() const { return Actuators.CreateConstIterator(); }
 	
+	/** If the avatar is a controller, then get the current score from the controller's player state. */
 	virtual float GetReward() const;
+
+	/** The agent is done if its avatar has been destroyed and it can't request a new one. */
 	virtual bool IsDone() const;
 
+	/** Get the actuator with the given ID if this agent has it. */
 	UMLAdapterActuator* GetActuator(const uint32 ActuatorID) 
 	{ 
 		UMLAdapterActuator** FoundActuator = Actuators.FindByPredicate([ActuatorID](const UMLAdapterActuator* Actuator) { return (Actuator->GetElementID() == ActuatorID); });
@@ -109,31 +135,53 @@ public:
 #endif // WITH_GAMEPLAY_DEBUGGER
 
 protected:
+
+	/** When the agent's avatar is destroyed, we need to cleanup callbacks and references to the avatar. Will request a new avatar if AgentConfig.bAutoRequestNewAvatarUponClearingPrev is true. */
 	UFUNCTION()
 	virtual void OnAvatarDestroyed(AActor* DestroyedActor);
 	
-	/** will be bound to UGameInstance.OnPawnControllerChanged if current avatar is a pawn or a controller */
+	/** Will be bound to UGameInstance.OnPawnControllerChanged if current avatar is a pawn or a controller. */
 	UFUNCTION()
 	void OnPawnControllerChanged(APawn* InPawn, AController* InController);
 
+	/** If the Pawn changed, we need to let all the sensors know. */
 	virtual void OnPawnChanged(APawn* NewPawn, AController* InController);
 
 	friend UMLAdapterSession;
+
+	/** Set this agent's ID to its new ID. */
 	void SetAgentID(FMLAdapter::FAgentID NewAgentID) { AgentID = NewAgentID; }
 
 public:
+
+	/** Retrieve all the sensor data from the last time Sense was called. */
 	void GetObservations(FMLAdapterMemoryWriter& Ar);
+
+	/** Get this agent's current config. This will not be accurate if the agent was spawned from a blueprint. */
 	const FMLAdapterAgentConfig& GetConfig() const { return AgentConfig; }
+
+	/** Setup this agent's avatar, sensors, and actuators. Typically used for agents spawned via RPC. */
 	virtual void Configure(const FMLAdapterAgentConfig& NewConfig);
+
+	/** Get the overall action space of this agent based on all its actuators. Used to determine the necessary output shape of the ML model. */
 	virtual void GetActionSpaceDescription(FMLAdapterSpaceDescription& OutSpaceDesc) const;
+
+	/** Get the overall observation space of this agent based on all its sensors. Used to determine the necessary input shape of the ML model. */
 	virtual void GetObservationSpaceDescription(FMLAdapterSpaceDescription& OutSpaceDesc) const;
 
+	/** Get the session that this agent belongs to. */
 	UMLAdapterSession& GetSession();
 	
+	/** Returns true is the given avatar can be controlled by this agent. */
 	virtual bool IsSuitableAvatar(AActor& InAvatar) const;
+
+	/** Sets the avatar for this agent and all of its sensors and actuators. Registers callbacks. */
 	virtual void SetAvatar(AActor* InAvatar);
+
+	/** Get the avatar this agent is controlling. */
 	AActor* GetAvatar() const { return Avatar; }
 
+	/** Returns true if this agent has an avatar set. */
 	bool IsReady() const { return Avatar != nullptr; }
 
 protected:
@@ -160,8 +208,12 @@ private:
 	APawn* Pawn;
 
 	FMLAdapter::FAgentID AgentID;
+
 	FMLAdapterAgentConfig AgentConfig;
 
+	/** True if the agent ever had a Pawn. Otherwise, false. */
 	uint32 bEverHadAvatar : 1;
+
+	/** True if we have callbacks registered. */
 	uint32 bRegisteredForPawnControllerChange : 1;
 };
diff --git a/Engine/Plugins/AI/MLAdapter/Source/MLAdapter/Public/Agents/MLAdapterAgentElement.h b/Engine/Plugins/AI/MLAdapter/Source/MLAdapter/Public/Agents/MLAdapterAgentElement.h
--- a/Engine/Plugins/AI/MLAdapter/Source/MLAdapter/Public/Agents/MLAdapterAgentElement.h
+++ b/Engine/Plugins/AI/MLAdapter/Source/MLAdapter/Public/Agents/MLAdapterAgentElement.h
@@ -13,6 +13,9 @@
 class UMLAdapterAgent;
 struct FMLAdapterDescription;
 
+/**
+ * An agent element is any object that can be attached to an agent. Base class for UMLAdapterSensor and UMLAdapterActuator.
+ */
 UCLASS(abstract)
 class MLADAPTER_API UMLAdapterAgentElement : public UObject
 {
@@ -21,32 +24,43 @@ public:
 	UMLAdapterAgentElement(const FObjectInitializer& ObjectInitializer = FObjectInitializer::Get());
 	virtual void PostInitProperties() override;
 
+	/** Get the agent this element is attached to. */
 	const UMLAdapterAgent& GetAgent() const;
+
+	/** Get the ID of this element. */
 	uint32 GetElementID() const { return ElementID; }
+
+	/** Get nickname of this element. */
 	const FString& GetNickname() const { return Nickname; }
 
-	/** @return owning agent's avatar */
+	/** @return The owning agent's avatar. */
 	AActor* GetAvatar() const;
 
-	/** @return the pawn associated with the owning agent. If owning agent's avatar 
-	 *	is a pawn then that gets retrieved, if not the function will check if 
-	 *	it's a controller and if so retrieve its pawn */
+	/**
+	 * @return The pawn associated with the owning agent. If owning agent's avatar is a pawn then that gets retrieved,
+	 * if not the function will check if it's a controller and if so retrieve its pawn.
+	 */
 	APawn* GetPawnAvatar() const;
 
-	/** @return the controller associated with the owning agent. If owning agent's 
-	 *	avatar is a controller then that gets retrieved, if not the function will 
-	 *	check if it's a pawn and if so retrieve its controller */
+	/**
+	 * @return The controller associated with the owning agent. If owning agent's avatar is a controller then that
+	 * gets retrieved, if not the function will check if it's a pawn and if so retrieve its controller.
+	 */
 	AController* GetControllerAvatar() const;
 
-	/** Fetches both the pawn and the controller associated with the current agent.
-	 *	It's like both calling @see GetPawnAvatar and @see GetControllerAvatar 
-	 *	@return true if at least one of the fetched pair is non-null */
+	/**
+	 * Fetches both the pawn and the controller associated with the current agent.
+	 * It's like both calling @see GetPawnAvatar and @see GetControllerAvatar.
+	 * @return true if at least one of the fetched pair is non-null.
+	 */
 	bool GetPawnAndControllerAvatar(APawn*& OutPawn, AController*& OutController) const;
 
 	virtual FString GetDescription() const { return Description; }
 
-	/** Called before object's destruction. Can be called as part of new agent 
-	 *	config application when old actuator get destroyed */
+	/**
+	 * Called before object's destruction. Can be called as part of new agent config application when old actuator get
+	 * destroyed.
+	 */
 	virtual void Shutdown() {}
 	
 	void SetNickname(const FString& NewNickname) { Nickname = NewNickname; }
@@ -64,24 +78,24 @@ protected:
 	virtual TSharedPtr<FMLAdapter::FSpace> ConstructSpaceDef() const PURE_VIRTUAL(UMLAdapterAgentElement::ConstructSpaceDef, return MakeShareable(new FMLAdapter::FSpace_Dummy()); );
 
 protected:
-	// can be queried by remote clients 
+	/** Can be queried by remote clients. */
 	FString Description;
 
 	TSharedRef<FMLAdapter::FSpace> SpaceDef;
 
-	/** @note this is not a common counter, meaning Sensors and Actuators (for 
-	 *	example) track the ID separately */
+	/** @note This is not a common counter, meaning Sensors and Actuators (for example) track the ID separately. */
 	UPROPERTY()
 	uint32 ElementID;
 
-	/** User-configured name for this element, mostly for debugging purposes but 
-	 *	comes in handy when fetching observation/action spaces descriptions.
-	 *	Defaults to UnrealEngine instance name*/
+	/**
+	 * User-configured name for this element, mostly for debugging purposes but comes in handy when fetching
+	 * observation/action spaces descriptions. Defaults to UnrealEngine instance name.
+	 */
 	UPROPERTY()
 	FString Nickname; 
 
 #if WITH_GAMEPLAY_DEBUGGER
-	// displayed in debugging tools and logging
+	/** Displayed in debugging tools and logging. */
 	mutable FString DebugRuntimeString;
 #endif // WITH_GAMEPLAY_DEBUGGER
 };
diff --git a/Engine/Plugins/AI/MLAdapter/Source/MLAdapter/Public/Agents/MLAdapterAgent_Inference.h b/Engine/Plugins/AI/MLAdapter/Source/MLAdapter/Public/Agents/MLAdapterAgent_Inference.h
--- a/Engine/Plugins/AI/MLAdapter/Source/MLAdapter/Public/Agents/MLAdapterAgent_Inference.h
+++ b/Engine/Plugins/AI/MLAdapter/Source/MLAdapter/Public/Agents/MLAdapterAgent_Inference.h
@@ -2,6 +2,7 @@
 
 #pragma once
 
+#include "UObject/SoftObjectPath.h"
 #include "Agents/MLAdapterAgent.h"
 #include "NeuralNetwork.h"
 #include "MLAdapterAgent_Inference.generated.h"
@@ -9,20 +10,24 @@
 /**
  * Inference agents have a neural network that can process senses and output actuations via their Think method. You 
  * can create a blueprint of this class to easily wire-up an agent that functions entirely inside the Unreal Engine.
- * Typically, this class would be used in conjunction with MLAdapterNoRPCManager.
+ * @see UMLAdapterNoRPCManager - we do not usually need to open the RPC interface if we are exclusively using inference agents.
  */
 UCLASS(Blueprintable, EditInlineNew)
 class MLADAPTER_API UMLAdapterAgent_Inference : public UMLAdapterAgent
 {
 	GENERATED_BODY()
 public:
+	/** Tries to load the neural network from the NeuralNetworkPath. */
 	virtual void PostInitProperties() override;
 
+	/** Process observations from the sensors with the neural network and populate the actuators with data. */
 	virtual void Think(const float DeltaTime) override;
 
+	/** The path to the neural network asset to be loaded. */
 	UPROPERTY(EditDefaultsOnly, Category = MLAdapter, meta = (AllowedClasses = "NeuralNetwork"))
 	FSoftObjectPath NeuralNetworkPath;
 
+	/** The neural network controlling this agent. */
 	UPROPERTY()
 	UNeuralNetwork* Brain;
 };
diff --git a/Engine/Plugins/AI/MLAdapter/Source/MLAdapter/Public/MLAdapterLibrarian.h b/Engine/Plugins/AI/MLAdapter/Source/MLAdapter/Public/MLAdapterLibrarian.h
--- a/Engine/Plugins/AI/MLAdapter/Source/MLAdapter/Public/MLAdapterLibrarian.h
+++ b/Engine/Plugins/AI/MLAdapter/Source/MLAdapter/Public/MLAdapterLibrarian.h
@@ -11,33 +11,96 @@
 #include "Agents/MLAdapterAgent.h"
 #include "MLAdapterLibrarian.generated.h"
 
+/**
+ * The FMLAdapterLibrarian discovers all classes that derived from UMLAdapterAgent, UMLAdapterSensor, or UMLAdapterActuator. 
+ * It provides functionality for finding those classes by name, so that remote clients can spawn them via their name.
+ */
 USTRUCT()
 struct FMLAdapterLibrarian
 {
 	GENERATED_BODY()
-		
+	
+	/**
+	* Register all the base classes and derived classes that inherit from UMLAdapterAgent, UMLAdapterSensor, or
+	* UMLAdapterActuator.
+	*/
 	void GatherClasses();
+
+	/** Add the sensor class to the list of known sensors. */
 	void RegisterSensorClass(const TSubclassOf<UMLAdapterSensor>& Class);
+
+	/** Add the actuator class to the list of known actuators. */
 	void RegisterActuatorClass(const TSubclassOf<UMLAdapterActuator>& Class);
+
+	/** Add the agent class to the list of known agents. */
 	void RegisterAgentClass(const TSubclassOf<UMLAdapterAgent>& Class);
 
+	/** Add the given description to the mapped entry for the given function name. */
 	void AddRPCFunctionDescription(const FName FunctionName, FString&& Description);
 
+	/** Get a const iterator over the list of known sensors. */
 	TMap<uint32, TSubclassOf<UMLAdapterSensor> >::TConstIterator GetSensorsClassIterator() const { return KnownSensorClasses.CreateConstIterator(); }
+
+	/** Get a const iterator over the list of known actuators. */
 	TMap<uint32, TSubclassOf<UMLAdapterActuator> >::TConstIterator GetActuatorsClassIterator() const { return KnownActuatorClasses.CreateConstIterator(); }
+
+	/** Get a const iterator over the list of known agents. */
 	TArray<TSubclassOf<UMLAdapterAgent> >::TConstIterator GetAgentsClassIterator() const { return KnownAgentClasses.CreateConstIterator(); }
 
+	/**
+	 * Gets the function description for the given function name.
+	 * @param FunctionName The name of the function whose description we would like to get.
+	 * @param OutDescription The description will be returned here if found.
+	 * @return True if the description was found. Otherwise false.
+	 */
 	bool GetFunctionDescription(const FName FunctionName, FString& OutDescription) const;
+
+	/**
+	 * Gets the function description for the given function name.
+	 * @param FunctionName The name of the function whose description we would like to get.
+	 * @param OutDescription The description will be returned here if found.
+	 * @return True if the description was found. Otherwise false.
+	 */
 	inline bool GetFunctionDescription(const FString& FunctionName, FString& OutDescription) const;
+
+	/** Get a const iterator over the list of function descriptions. */
 	TMap<FName, FString>::TConstIterator GetFunctionDescriptionsIterator() const { return RPCFunctionDescriptions.CreateConstIterator(); }
 
+	/**
+	 * Gets the description for the given sensor name.
+	 * @param SensorName The name of the sensor whose description we would like to get.
+	 * @param OutDescription The description will be returned here if found.
+	 * @return True if the description was found. Otherwise false.
+	 */
 	bool GetSensorDescription(const FName SensorName, FString& OutDescription) const;
+
+	/**
+	 * Gets the description for the given actuator name.
+	 * @param ActuatorName The name of the actuator whose description we would like to get.
+	 * @param OutDescription The description will be returned here if found.
+	 * @return True if the description was found. Otherwise false.
+	 */
 	bool GetActuatorDescription(const FName ActuatorName, FString& OutDescription) const;
 
+	/**
+	 * Find the agent class with the given class name. Can be the full name or the substring after "UMLAdapterAgent_".
+	 * Will return UMLAdapterAgent::StaticClass() if not found.
+	 */
 	TSubclassOf<UMLAdapterAgent> FindAgentClass(const FName ClassName) const;
+
+	/**
+     * Find the sensor class with the given class name. Can be the full name or the substring after "UMLAdapterSensor_".
+	 * Will return UMLAdapterSensor::StaticClass() if not found.
+	 */
 	TSubclassOf<UMLAdapterSensor> FindSensorClass(const FName ClassName) const;
+
+	/**
+	 * Find the actuator class with the given class name. Can be the full name or the substring after "UMLAdapterActuator_".
+	 * Will return UMLAdapterActuator::StaticClass() if not found.
+	 */
 	TSubclassOf<UMLAdapterActuator> FindActuatorClass(const FName ClassName) const;
 
+	/** Gets the librarian owned by the singleton UMLAdapterManager instance. */
 	static const FMLAdapterLibrarian& Get();
 
 protected:
diff --git a/Engine/Plugins/AI/MLAdapter/Source/MLAdapter/Public/MLAdapterScribe.h b/Engine/Plugins/AI/MLAdapter/Source/MLAdapter/Public/MLAdapterScribe.h
--- a/Engine/Plugins/AI/MLAdapter/Source/MLAdapter/Public/MLAdapterScribe.h
+++ b/Engine/Plugins/AI/MLAdapter/Source/MLAdapter/Public/MLAdapterScribe.h
@@ -10,7 +10,7 @@
 #include <string>
 #include <map>
 
-/** Describes requested data in an rpclib-usable format */
+/** Describes requested data in an rpclib-usable format. */
 namespace FMLAdapterScribe
 {
 	std::vector<std::string> ToStringVector(const TArray<FString>& Array);
@@ -19,7 +19,12 @@ namespace FMLAdapterScribe
 	std::map<std::string, uint32> ListSensorTypes();
 	std::map<std::string, uint32> ListActuatorTypes();
 
-	// will search for the given names first in function names, then sensors, 
-	// actuators, agent classes. To be expanded
+	/**
+	* Gets the description for the given name. Searches the Librarian for registered names in the following order:
+	*	1. Functions
+	*	2. Sensors
+	*	3. Actuators
+	* @return The description if found. Otherwise, returns string literal "Not Found".
+	*/ 
 	std::string GetDescription(std::string const& ElementName);
 };
\ No newline at end of file
diff --git a/Engine/Plugins/AI/MLAdapter/Source/MLAdapter/Public/MLAdapterSpace.h b/Engine/Plugins/AI/MLAdapter/Source/MLAdapter/Public/MLAdapterSpace.h
--- a/Engine/Plugins/AI/MLAdapter/Source/MLAdapter/Public/MLAdapterSpace.h
+++ b/Engine/Plugins/AI/MLAdapter/Source/MLAdapter/Public/MLAdapterSpace.h
@@ -21,6 +21,7 @@ DECLARE_ENUM_TO_STRING(EMLAdapterSpaceType);
 
 namespace FMLAdapter
 {
+	/** Defines the numerical space of a sensor, actuator, or agent. Similar to OpenAI Gym's spaces. */
 	struct MLADAPTER_API FSpace : public IJsonable, public TSharedFromThis<FSpace>
 	{
 		virtual ~FSpace() {}
@@ -30,6 +31,7 @@ namespace FMLAdapter
 		virtual int32 Num() const { return 0; }
 	};
 
+	/** A discrete space contains a countable number of values. */
 	struct MLADAPTER_API FSpace_Discrete : public FSpace
 	{
 		uint32 Count;
@@ -51,6 +53,7 @@ namespace FMLAdapter
 		virtual int32 Num() const override { return Options.Num(); }
 	};
 
+	/** A continuous space that contains a number of ranges defined by the shape, whose values will fall inside the Low to High range. */
 	struct MLADAPTER_API FSpace_Box : public FSpace
 	{
 		TArray<uint32> Shape;
@@ -65,12 +68,14 @@ namespace FMLAdapter
 		static TSharedPtr<FSpace> Vector2D(float Low = -1.f, float High = 1.f) { return MakeShareable(new FSpace_Box({ 2 }, Low, High)); }
 	};
 
+	/** A placeholder shape - typically returned in invalid scenarios. */
 	struct MLADAPTER_API FSpace_Dummy : public FSpace_Box
 	{
 		FSpace_Dummy() : FSpace_Box({ 0 }) {}
 		virtual int32 Num() const override { return 0; }
 	};
 
+	/** Container for multiple subspaces. */
 	struct MLADAPTER_API FSpace_Tuple : public FSpace
 	{
 		TArray<TSharedPtr<FSpace> > SubSpaces;
@@ -114,6 +119,7 @@ struct FMLAdapterDescription
 	void Reset() { Data.Reset(); PrepData.Reset(); }
 
 	bool IsEmpty() const { return (Data.Num() == 0); }
+
 protected:
 	TArray<TPair<FString, FString>> Data;
 	TArray<FString> PrepData;
@@ -127,6 +133,7 @@ struct FMLAdapterSpaceDescription
 
 	FMLAdapterSpaceDescription& Add(FString Key, FMLAdapterDescription&& Element) { Data.Add(ValuePair(Key, Element)); return(*this); }
 	FMLAdapterSpaceDescription& Add(FString Key, FMLAdapterDescription Element) { Data.Add(ValuePair(Key, Element)); return(*this); }
+
 protected:
 	TArray<ValuePair> Data;
 };
diff --git a/Engine/Plugins/AI/MLAdapter/Source/MLAdapter/Public/MLAdapterTypes.h b/Engine/Plugins/AI/MLAdapter/Source/MLAdapter/Public/MLAdapterTypes.h
--- a/Engine/Plugins/AI/MLAdapter/Source/MLAdapter/Public/MLAdapterTypes.h
+++ b/Engine/Plugins/AI/MLAdapter/Source/MLAdapter/Public/MLAdapterTypes.h
@@ -31,7 +31,8 @@ namespace FMLAdapter
 	static const uint32 InvalidSensorID = 0;
 	static const uint32 InvalidActuatorID = 0;
 
-	/** MLAdapter-flavored new object creation. We're using this with all the objects 
+	/**
+	 *	MLAdapter-flavored new object creation. We're using this with all the objects 
 	 *	we intend to access from outside the game-thread (most notably from the 
 	 *	rpc calls). This gives us more control over objects' life cycle.
 	 *
diff --git a/Engine/Plugins/AI/MLAdapter/Source/MLAdapter/Public/Managers/MLAdapterManager.h b/Engine/Plugins/AI/MLAdapter/Source/MLAdapter/Public/Managers/MLAdapterManager.h
--- a/Engine/Plugins/AI/MLAdapter/Source/MLAdapter/Public/Managers/MLAdapterManager.h
+++ b/Engine/Plugins/AI/MLAdapter/Source/MLAdapter/Public/Managers/MLAdapterManager.h
@@ -23,7 +23,7 @@ class APlayerController;
 class UWorld;
 
 UENUM()
-enum class EMLAdapterServerMode
+enum class EMLAdapterServerMode : uint8
 {
 	Invalid,
 	Server,
@@ -32,7 +32,11 @@ enum class EMLAdapterServerMode
 	AutoDetect,
 };
 
-
+/**
+ * The manager of the MLAdapter system. Sets up the RPC server for communication with remote client. Based on settings,
+ * creates and ticks the UMLAdapterSession. Tells the UMLAdapterLibrarian to gather its classes.
+ * A singleton instance is setup automatically during OnPostEngineInit if this plugin is included.
+ */
 UCLASS(Transient)
 class MLADAPTER_API UMLAdapterManager : public UObject, public FTickableGameObject, public FSelfRegisteringExec
 {
@@ -59,47 +63,72 @@ public:
 	virtual bool Exec(class UWorld* InWorld, const TCHAR* Cmd, FOutputDevice& Ar) override;
 	// FExec end
 
+	/** Binds the manager's world, game mode, and editor delegates - called during PostInitProperties. */
 	virtual void BindToDelegates();
+
+	/** Cleans up all the manager's delegates - called during BeginDestroy. */
 	virtual void CleanUpDelegates();
 
+	/** Starts the RPC server if the world is a game world */
 	virtual void OnPostWorldInit(UWorld* World, const UWorld::InitializationValues);
+
+	/** Closes the current session. */
 	virtual void OnWorldCleanup(UWorld* World, bool bSessionEnded, bool bCleanupResources);
+
 	virtual void OnBeginPIE(const bool bIsSimulating);
 	virtual void OnEndPIE(const bool bIsSimulating);
+
+	/** Notifies the session that the game mode has initialized. */
 	virtual void OnGameModeInitialized(AGameModeBase* GameMode);
+
+	/** Notifies the session a player has connected. */
 	virtual void OnGameModePostLogin(AGameModeBase* GameMode, APlayerController* NewPlayer);
-	/** @note that this might not get called at all if the project's game mode 
-	 *	doesn't extend AGameMode*/
+
+	/** @note This might not get called at all if the project's game mode doesn't extend AGameMode. */
 	virtual void OnGameModeMatchStateSet(FName MatchState);
 
-	/** If a server is already running it will be shut down before the new instance gets created
-	 *	@param ServerThreads best set at the number of external clients that are going to be connecting */
+	/**
+	 * Starts a new RPC server. If a server is already running, it will be shut down before the new instance gets created.
+	 * @param ServerThreads best set at the number of external clients that are going to be connecting
+	 */
 	virtual void StartServer(uint16 Port, EMLAdapterServerMode InMode = EMLAdapterServerMode::AutoDetect, uint16 ServerThreads = 1);
+	
+	/** Stop the RPC server. */
 	virtual void StopServer();
+
+	/** True if there is an RPC server currently running. */
 	virtual bool IsRunning() const;
 
+	/** "Server" in this context means UnrealEngine game server, not RPC server. */
 	virtual void ConfigureAsServer(FRPCServer& Server);
-	/** "Client" in this context means UnrealEngine game client, not RPC client */
+
+	/** "Client" in this context means UnrealEngine game client, not RPC client. */
 	virtual void ConfigureAsClient(FRPCServer& Server);
-	/** Essentially calls both the server and client versions */
+
+	/** Essentially calls both the server and client versions. */
 	virtual void ConfigureAsStandalone(FRPCServer& Server);
 
-	/** if given World doesn't have an AI system this call results in creating one */
+	/** If given World doesn't have an AI system, this call results in creating one. */
 	virtual void EnsureAISystemPresence(UWorld& World);
 
-	/** if given World doesn't have a Navigation system instance this call results in creating one */
+	/** If given World doesn't have a Navigation system instance, this call results in creating one. */
 	virtual void EnsureNavigationSystemPresence(UWorld& World);
 
 	virtual UMLAdapterSession* CreateNewSession();
 	virtual void SetSession(UMLAdapterSession* NewSession);
 	virtual void CloseSession(UMLAdapterSession& InSession);
 
-	/** Returns current session. If one doesn't exist, it gets created. */
+	/** Returns the current session. If one doesn't exist, it gets created. */
 	virtual UMLAdapterSession& GetSession();
 	bool HasSession() const { return IsValid(Session) && (Session->IsUnreachable() == false); }
 
+	/** Register a sensor class with this manager's librarian. */
 	void RegisterSensorClass(const TSubclassOf<UMLAdapterSensor>& Class) { Librarian.RegisterSensorClass(Class); }
+
+	/** Register an actuator class with this manager's librarian. */
 	void RegisterActuatorClass(const TSubclassOf<UMLAdapterActuator>& Class) { Librarian.RegisterActuatorClass(Class); }
+
+	/** Register an agent class with this manager's librarian. */
 	void RegisterAgentClass(const TSubclassOf<UMLAdapterAgent>& Class) { Librarian.RegisterAgentClass(Class); }
 
 	virtual void ResetWorld();
@@ -108,11 +137,16 @@ public:
 	FOnGenericRPCServerDelegate& GetOnAddClientFunctions() { return OnAddClientFunctions; }
 	FOnGenericRPCServerDelegate& GetOnAddServerFunctions() { return OnAddServerFunctions; }
 
+	/** Get the current manager instance. */
 	FORCEINLINE static UMLAdapterManager& Get();
+
+	/** Returns true if the manager instance exists. */
 	FORCEINLINE static bool IsReady();
 
+	/** Get this manager's librarian. */
 	const FMLAdapterLibrarian& GetLibrarian() const { return Librarian; }
 
+	/** Returns true if this manager is not being ticked manually by the remote client. */
 	bool IsWorldRealTime() const { return (bTickWorldManually == false); }
 
 	FOnGenericEvent& GetOnCurrentSessionChanged() { return OnCurrentSessionChanged; }
@@ -153,7 +187,7 @@ protected:
 	uint32 bTickWorldManually : 1;
 	
 	/** is the manager is in 'manual ticking mode' (where external client is 
-	 *	responsible for progressing the world sim by calling 'request_world_tick' 
+	 * responsible for progressing the world sim by calling 'request_world_tick' 
 	 *	function) the simulation will progress by StepsRequested ticks before pausing */
 	int32 StepsRequested = 0;
 
diff --git a/Engine/Plugins/AI/MLAdapter/Source/MLAdapter/Public/Managers/MLAdapterNoRPCManager.h b/Engine/Plugins/AI/MLAdapter/Source/MLAdapter/Public/Managers/MLAdapterNoRPCManager.h
--- a/Engine/Plugins/AI/MLAdapter/Source/MLAdapter/Public/Managers/MLAdapterNoRPCManager.h
+++ b/Engine/Plugins/AI/MLAdapter/Source/MLAdapter/Public/Managers/MLAdapterNoRPCManager.h
@@ -7,7 +7,8 @@
 #include "MLAdapterNoRPCManager.generated.h"
 
 /**
- * No RPC manager won't start RPC server and will immediately start a session and spawn the default agent.
+ * UMLAdapterNoRPCManager won't start the RPC server and will immediately start a session and spawn the default agent. This is
+ * intended to be used when there will only be a single inference agent that exists for the life of the game.
  */
 UCLASS()
 class MLADAPTER_API UMLAdapterNoRPCManager : public UMLAdapterManager
diff --git a/Engine/Plugins/AI/MLAdapter/Source/MLAdapter/Public/Sensors/MLAdapterSensor.h b/Engine/Plugins/AI/MLAdapter/Source/MLAdapter/Public/Sensors/MLAdapterSensor.h
--- a/Engine/Plugins/AI/MLAdapter/Source/MLAdapter/Public/Sensors/MLAdapterSensor.h
+++ b/Engine/Plugins/AI/MLAdapter/Source/MLAdapter/Public/Sensors/MLAdapterSensor.h
@@ -10,6 +10,7 @@
 class UMLAdapterAgent;
 struct FMLAdapterDescription;
 
+/** Controls the frequency that a sensor should tick at. */
 UENUM()
 enum class EMLAdapterTickPolicy : uint8
 {
@@ -19,16 +20,20 @@ enum class EMLAdapterTickPolicy : uint8
 	Never
 };
 
+/** Allows an agent to sense information about the game world or itself. */
 UCLASS(Abstract, Blueprintable, EditInlineNew)
 class MLADAPTER_API UMLAdapterSensor : public UMLAdapterAgentElement
 {
     GENERATED_BODY()
+
 public:
     UMLAdapterSensor(const FObjectInitializer& ObjectInitializer = FObjectInitializer::Get());
 	virtual void PostInitProperties() override;
 
-	/** Called before actuator's destruction. Can be called as part of new 
-	 *	agent config application when old actuator get destroyed */
+	/**
+	 * Called before actuator's destruction. Can be called as part of new agent config application when old actuator
+	 * get destroyed.
+	 */
 	virtual void Configure(const TMap<FName, FString>& Params) override;
 
 	virtual void OnAvatarSet(AActor* Avatar) override;
@@ -39,14 +44,12 @@ public:
 	bool IsConfiguredForAgent(const UMLAdapterAgent& Agent) const;
 	bool IsPolling() const { return bIsPolling; }
 
-	/** @return True if config was successful. Only in that case will the sensor
-	 *		instance be added to Agent's active sensors */
+	/** True if config was successful. Only then will the sensor instance be added to agent's active sensors. */
 	virtual bool ConfigureForAgent(UMLAdapterAgent& Agent);
 
 	void OnPawnChanged(APawn* OldPawn, APawn* NewPawn);
 
-	/** Called for every sense, regardless of whether it's a polling-type or not.
-	 */
+	/** Called for every sense, regardless of whether it's a polling-type or not. */
 	void Sense(const float DeltaTime);
 
 	/** 
@@ -55,21 +58,26 @@ public:
 	 */
 	virtual void GetObservations(FMLAdapterMemoryWriter& Ar) PURE_VIRTUAL(UMLAdapterSensor::GetObservations, );
 
-protected:	
+protected:
+
+	/** Clean up any references to the Pawn or its components. */
 	virtual void ClearPawn(APawn& InPawn);
 
-	/** called from Sense based on TickPolicy */
+	/** Called from Sense based on TickPolicy. */
 	virtual void SenseImpl(const float DeltaTime) {}
 
 protected:
 	FMLAdapter::FAgentID AgentID;
 
+	/** @todo this is not currently referenced anywhere */
 	UPROPERTY(EditDefaultsOnly, Category=MLAdapter)
 	uint32 bRequiresPawn : 1;
 
+	/** @todo this is not currently referenced anywhere */
 	UPROPERTY(EditDefaultsOnly, Category = MLAdapter)
 	uint32 bIsPolling : 1;
 
+	/** Determines what frequency this sensor ticks at. */
 	UPROPERTY(EditDefaultsOnly, Category = MLAdapter)
 	EMLAdapterTickPolicy TickPolicy;
 
@@ -84,6 +92,7 @@ protected:
 	};
 	FTicksOrSeconds TickEvery;
 
+	/**  This needs to be locked anytime the sensor is accessing shared state internally. */
 	mutable FCriticalSection ObservationCS;
 private:
 	int32 AccumulatedTicks;
diff --git a/Engine/Plugins/AI/MLAdapter/Source/MLAdapter/Public/Sensors/MLAdapterSensor_AIPerception.h b/Engine/Plugins/AI/MLAdapter/Source/MLAdapter/Public/Sensors/MLAdapterSensor_AIPerception.h
--- a/Engine/Plugins/AI/MLAdapter/Source/MLAdapter/Public/Sensors/MLAdapterSensor_AIPerception.h
+++ b/Engine/Plugins/AI/MLAdapter/Source/MLAdapter/Public/Sensors/MLAdapterSensor_AIPerception.h
@@ -9,7 +9,8 @@
 class UAIPerceptionComponent;
 class AActor;
 
-/** When applied to a player controller will create an AIPerception component for 
+/**
+ *	When applied to a player controller will create an AIPerception component for 
  *	that player and plug it into the AIPerceptionSystem. The sensor will report 
  *	information gathered by the perception system on the behalf of this agent.
  *	@see UAIPerceptionComponent::ProcessStimuli
@@ -75,10 +76,12 @@ protected:
 	UPROPERTY()
 	UAIPerceptionComponent* PerceptionComponent;
 
-	/** When set to true will only gather perception "delta" meaning consecutive
+	/**
+	 *	When set to true will only gather perception "delta" meaning consecutive
 	 *	updates will consist of new perception information. Defaults to "false" 
 	 *	which means that every update all of data contained by the PerceptionComponent 
-	 *	will be "sensed" */
+	 *	will be "sensed"
+	 */
 	UPROPERTY()
 	bool bSenseOnlyChanges = false;
 
diff --git a/Engine/Plugins/AI/MLAdapter/Source/MLAdapter/Public/Sensors/MLAdapterSensor_EnhancedInput.h b/Engine/Plugins/AI/MLAdapter/Source/MLAdapter/Public/Sensors/MLAdapterSensor_EnhancedInput.h
--- a/Engine/Plugins/AI/MLAdapter/Source/MLAdapter/Public/Sensors/MLAdapterSensor_EnhancedInput.h
+++ b/Engine/Plugins/AI/MLAdapter/Source/MLAdapter/Public/Sensors/MLAdapterSensor_EnhancedInput.h
@@ -26,7 +26,6 @@ protected:
 
 	TMap<FString, int32> InputStateIndices;
 
-public:
 	UPROPERTY(EditDefaultsOnly, BlueprintReadOnly, Category = MLAdapter)
 	TArray<TObjectPtr<UInputAction>> TrackedActions;
 
diff --git a/Engine/Plugins/AI/MLAdapter/Source/MLAdapter/Public/Sensors/MLAdapterSensor_Input.h b/Engine/Plugins/AI/MLAdapter/Source/MLAdapter/Public/Sensors/MLAdapterSensor_Input.h
--- a/Engine/Plugins/AI/MLAdapter/Source/MLAdapter/Public/Sensors/MLAdapterSensor_Input.h
+++ b/Engine/Plugins/AI/MLAdapter/Source/MLAdapter/Public/Sensors/MLAdapterSensor_Input.h
@@ -14,9 +14,9 @@ class FViewport;
 class UGameViewportClient;
 
 
-/** Note that this sensor doesn't buffer input state between GetObservations call
- *	@todo a child class could easily do that by overriding OnInputKey/OnInputAxis and 
- *		GetObservations
+/**
+ *	Note that this sensor doesn't buffer input state between GetObservations call
+ *	@todo a child class could easily do that by overriding OnInputKey/OnInputAxis and GetObservations
  */
 UCLASS(Blueprintable)
 class MLADAPTER_API UMLAdapterSensor_Input : public UMLAdapterSensor
diff --git a/Engine/Plugins/AI/MLAdapter/Source/MLAdapter/Public/Sensors/MLAdapterSensor_Movement.h b/Engine/Plugins/AI/MLAdapter/Source/MLAdapter/Public/Sensors/MLAdapterSensor_Movement.h
--- a/Engine/Plugins/AI/MLAdapter/Source/MLAdapter/Public/Sensors/MLAdapterSensor_Movement.h
+++ b/Engine/Plugins/AI/MLAdapter/Source/MLAdapter/Public/Sensors/MLAdapterSensor_Movement.h
@@ -6,7 +6,7 @@
 #include "MLAdapterTypes.h"
 #include "MLAdapterSensor_Movement.generated.h"
 
-
+/** Allows an agent to sense its avatar's location and velocity. */
 UCLASS(Blueprintable)
 class MLADAPTER_API UMLAdapterSensor_Movement : public UMLAdapterSensor
 {
diff --git a/Engine/Plugins/AI/MLAdapter/Source/MLAdapter/Public/Sessions/MLAdapterLocalDataCollectionSession.h b/Engine/Plugins/AI/MLAdapter/Source/MLAdapter/Public/Sessions/MLAdapterLocalDataCollectionSession.h
--- a/Engine/Plugins/AI/MLAdapter/Source/MLAdapter/Public/Sessions/MLAdapterLocalDataCollectionSession.h
+++ b/Engine/Plugins/AI/MLAdapter/Source/MLAdapter/Public/Sessions/MLAdapterLocalDataCollectionSession.h
@@ -7,7 +7,7 @@
 #include "MLAdapterLocalDataCollectionSession.generated.h"
 
 /**
- * Collects data from Agents' Sensors and writes them to a file for offline processing. Only works with a locally
+ * Collects data from agents' sensors and writes them to a file for offline processing. Only works with a locally
  * controlled, single-player game.
  */
 UCLASS(Blueprintable, EditInlineNew)
diff --git a/Engine/Plugins/AI/MLAdapter/Source/MLAdapter/Public/Sessions/MLAdapterSession.h b/Engine/Plugins/AI/MLAdapter/Source/MLAdapter/Public/Sessions/MLAdapterSession.h
--- a/Engine/Plugins/AI/MLAdapter/Source/MLAdapter/Public/Sessions/MLAdapterSession.h
+++ b/Engine/Plugins/AI/MLAdapter/Source/MLAdapter/Public/Sessions/MLAdapterSession.h
@@ -30,15 +30,21 @@ enum class EMLAdapterSimState : uint8
 };
 
 
+/**
+ * Container for agents that exist in the world. Ticks the agents. Finds avatars for the agents.
+ */
 UCLASS()
 class MLADAPTER_API UMLAdapterSession : public UObject
 {
 	GENERATED_BODY()
+
 public:
 	DECLARE_MULTICAST_DELEGATE_TwoParams(FOnAgentAvatarChangedDelegate, UMLAdapterAgent& /*Agent*/, AActor* /*OldAvatar*/);
 	DECLARE_MULTICAST_DELEGATE_OneParam(FOnBeginAgentRemove, UMLAdapterAgent& /*Agent*/);
 
+	/** Get this session's cached world. */
 	virtual UWorld* GetWorld() const override { return CachedWorld; }
+
 	virtual void PostInitProperties() override;
 
 	virtual void BeginDestroy() override;
@@ -46,38 +52,54 @@ public:
 	// @todo this needs further consideration.
 
 	UGameInstance* GetGameInstance() const { return CachedWorld ? CachedWorld->GetGameInstance() : nullptr; }
-	/** NewWorld might be null */
+	
+	/** Sets the world that this session will use to find avatars for agents */
 	virtual void SetWorld(UWorld* NewWorld);
 
+	/** If there are agents waiting for an avatar, then check if a newly spawned actor fits the agent. */
 	virtual void OnActorSpawned(AActor* InActor);
+
 	virtual void OnPostWorldInit(UWorld& World);
 	virtual void OnWorldCleanup(UWorld& World, bool bSessionEnded, bool bCleanupResources);
 	virtual void OnGameModeInitialized(AGameModeBase& GameModeBase);
 	virtual void OnGameModeMatchStateSet(FName InMatchState);
 	virtual void OnGameModePostLogin(AGameModeBase* GameMode, APlayerController* NewPlayer);
 
+	/** Mark the session as active. */
 	virtual void Open();
+
+	/** Mark the session as inactive and cleanup all agents. */
 	virtual void Close();
 
+	/** Call all agents' Sense(), Think(), and Act() methods. */
 	virtual void Tick(float DeltaTime);
 
-	/** This is where Session can add Authority-side-specific functions by calling 
-	 *	UMLAdapterManager::Get().AddServerFunctionBind */
+	/**
+	 *	This is where Session can add Authority-side-specific functions by calling
+	 *	UMLAdapterManager::Get().AddServerFunctionBind
+	 */
 	virtual void ConfigureAsServer();
 
-	/** This is where Session can add Client-side-specific functions by calling 
-	 *	UMLAdapterManager::Get().AddClientFunctionBind. "Client" in this context means 
-	 *	UnrealEngine game client, not RPC client */
+	/**
+	 *	This is where Session can add Client-side-specific functions by calling 
+	 *	UMLAdapterManager::Get().AddClientFunctionBind. "Client" in this context means UnrealEngine game client, not
+	 *	RPC client.
+	 */
 	virtual void ConfigureAsClient();
 
-	/** using FMLAdapter::InvalidAgentID for AgentID will reset all agents */
+	/** Using FMLAdapter::InvalidAgentID for AgentID will reset all agents. */
 	virtual void ResetWorld(FMLAdapter::FAgentID AgentID = FMLAdapter::InvalidAgentID);
 
+	/** A session is done if the SimulationState == Finished or the game mode's match has ended. */
 	bool IsDone() const;
+
+	/** A session is ready if the SimulationState == InProgress or the game mode's match has started but not ended. */
 	bool IsReady() const;
 
+	/** Get the world time in seconds at the most recent time this session ticked. */
 	float GetTimestamp() const { return LastTimestamp; }
 
+	/** If true, the world's ticking will be controlled by the remote client. */
 	void SetManualWorldTickEnabled(bool bEnable);
 
 	FOnAgentAvatarChangedDelegate GetOnAgentAvatarChanged() { return OnAgentAvatarChanged; }
@@ -86,43 +108,61 @@ public:
 	//----------------------------------------------------------------------//
 	// Agent/Avatar management 
 	//----------------------------------------------------------------------//
+
+	/** Add a default agent as specified in the UMLAdapterSettings. */
 	FMLAdapter::FAgentID AddAgent();
+
+	/** Add an agent using the specified config. */
 	FMLAdapter::FAgentID AddAgent(const FMLAdapterAgentConfig& InConfig);
 
-	/** @return next valid agent ID. Note that the return value might be equal to 
-	 *	ReferenceAgentID if there's only one agent. Will be FMLAdapter::InvalidAgentId 
-	 *	if no agents registered */
+	/**
+	 *	@return Next valid agent ID. Note that the return value might be equal to ReferenceAgentID if there's only one
+	 *	agent. Will be FMLAdapter::InvalidAgentId if no agents registered.
+	 */
 	FMLAdapter::FAgentID GetNextAgentID(FMLAdapter::FAgentID ReferenceAgentID) const;
+
+	/** Returns the agent corresponding to the given ID. Returns nullptr if the ID is invalid. */
 	UMLAdapterAgent* GetAgent(FMLAdapter::FAgentID AgentID);
 
+	/** Remove the agent with the given ID from this session. */
 	void RemoveAgent(FMLAdapter::FAgentID AgentID);
 
+	/** Returns true if the agent with the given ID has a valid avatar assigned. */
 	bool IsAgentReady(FMLAdapter::FAgentID AgentID) const;
 
 	/** Finds avatar in given World for every avatar-less agent in AwaitingAvatar */
 	void FindAvatars(UWorld& World);
 
-	/** Processes Agents and removes all agent avatars belonging to World. 
-	 *	If World is null the function will remove all avatars */
+	/**
+	 *	Processes Agents and removes all agent avatars belonging to World. 
+	 *	If World is null the function will remove all avatars.
+	 */
 	void RemoveAvatars(UWorld* World);
 
-	/** Finds a suitable avatar in InWorld (or CachedWorld, if InWorld is null) 
+	/**
+	 *	Finds a suitable avatar in InWorld (or CachedWorld, if InWorld is null) 
 	 *	for given agent, as specified by FMLAdapterAgentConfig.AvatarClass
 	 *	and confirmed by Agent->IsSuitableAvatar call. If no suitable avatar is 
 	 *	found this agent will be added to "waiting list" (AwaitingAvatar)
 	 *	@param bForceSearch if true will ignore whether the Agent is already waiting in 
 	 *		AwaitingAvatar and will perform the search right away. Note that the Agent 
 	 *		might still end up in AwaitingAvatar if there's no suitable avatars available 
-	 *	@return true if an avatar has been assigned. False otherwise.*/
+	 *	@return True if an avatar has been assigned. False otherwise.*/
 	virtual bool RequestAvatarForAgent(UMLAdapterAgent& Agent, UWorld* InWorld = nullptr, const bool bForceSearch = false);
 
+	/** Calls RequestAvatarForAgent(UMLAdapterAgent& Agent, ...) after looking up the agent by its ID. */
 	bool RequestAvatarForAgent(FMLAdapter::FAgentID& AgentID, UWorld* InWorld = nullptr);
 
+	/** Sets a new avatar on the agent after clearing any existing avatar. */
 	void BindAvatar(UMLAdapterAgent& Agent, AActor& Avatar);
+
+	/** Remove the existing avatar from the given agent. */
 	void ClearAvatar(UMLAdapterAgent& Agent);
 
+	/** @todo this technically returns the count of all agents that ever existed given that we are not removing them yet. */
 	int32 GetAgentsCount() const { return Agents.Num(); }
 
+	/** Find an agent that is controlling the given avatar if one exists. Returns nullptr if not. */
 	const UMLAdapterAgent* FindAgentByAvatar(AActor& Avatar) const;
 
 	//----------------------------------------------------------------------//
