MLAdapter: fix manual tick mode reset not working

#rnx

diff --git a/Engine/Plugins/AI/MLAdapter/Source/MLAdapter/Private/Managers/MLAdapterManager.cpp b/Engine/Plugins/AI/MLAdapter/Source/MLAdapter/Private/Managers/MLAdapterManager.cpp
--- a/Engine/Plugins/AI/MLAdapter/Source/MLAdapter/Private/Managers/MLAdapterManager.cpp
+++ b/Engine/Plugins/AI/MLAdapter/Source/MLAdapter/Private/Managers/MLAdapterManager.cpp
@@ -479,6 +479,14 @@ void UMLAdapterManager::OnGameModeMatchStateSet(FName MatchState)
 
 void UMLAdapterManager::ResetWorld()
 {
+#if WITH_EDITORONLY_DATA
+	// Reset logic won't work if the game thread is paused so we need to resume it
+	if (GIntraFrameDebuggingGameThread)
+	{
+		GIntraFrameDebuggingGameThread = false;
+	}
+#endif // WITH_EDITORONLY_DATA
+
 	if (LastActiveWorld)
 	{
 		AGameModeBase* GameMode = LastActiveWorld->GetAuthGameMode<AGameModeBase>();
