MLAdapter: fix bug in EnhancedInput actuator

#rnx

diff --git a/Engine/Plugins/AI/MLAdapter/Source/MLAdapter/Private/Actuators/MLAdapterActuator_EnhancedInput.cpp b/Engine/Plugins/AI/MLAdapter/Source/MLAdapter/Private/Actuators/MLAdapterActuator_EnhancedInput.cpp
--- a/Engine/Plugins/AI/MLAdapter/Source/MLAdapter/Private/Actuators/MLAdapterActuator_EnhancedInput.cpp
+++ b/Engine/Plugins/AI/MLAdapter/Source/MLAdapter/Private/Actuators/MLAdapterActuator_EnhancedInput.cpp
@@ -35,6 +35,11 @@ void UMLAdapterActuator_EnhancedInput::Act(const float DeltaTime)
 
 	FScopeLock Lock(&ActionCS);
 
+	if (InputData.IsEmpty())
+	{
+		return;
+	}
+
 	checkf(InputData.Num() % SpaceDef->Num() == 0, TEXT("InputData contains an unexpected number of elements"));
 
 	if (UEnhancedInputLocalPlayerSubsystem* Subsystem = ULocalPlayer::GetSubsystem<UEnhancedInputLocalPlayerSubsystem>(PC->GetLocalPlayer()))
