MLAdapter: LWC fixes

#rnx

diff --git a/Engine/Plugins/AI/MLAdapter/Source/MLAdapter/Private/Actuators/MLAdapterActuator_Camera.cpp b/Engine/Plugins/AI/MLAdapter/Source/MLAdapter/Private/Actuators/MLAdapterActuator_Camera.cpp
--- a/Engine/Plugins/AI/MLAdapter/Source/MLAdapter/Private/Actuators/MLAdapterActuator_Camera.cpp
+++ b/Engine/Plugins/AI/MLAdapter/Source/MLAdapter/Private/Actuators/MLAdapterActuator_Camera.cpp
@@ -87,11 +87,17 @@ void UMLAdapterActuator_Camera::DigestInputData(FMLAdapterMemoryReader& ValueStr
 	FScopeLock Lock(&ActionCS);
 	if (bVectorMode)
 	{
-		ValueStream << HeadingVector;
+		FVector3f HeadingVector3f;
+		ValueStream << HeadingVector3f;
+		HeadingVector = HeadingVector3f;
 	}
 	else
 	{
-		ValueStream << HeadingRotator.Pitch << HeadingRotator.Yaw;
+		float Pitch;
+		float Yaw;
+		ValueStream << Pitch << Yaw;
+		HeadingRotator.Pitch = Pitch;
+		HeadingRotator.Yaw = Yaw;
 	}
 }
 
diff --git a/Engine/Plugins/AI/MLAdapter/Source/MLAdapter/Private/Sensors/MLAdapterSensor_AIPerception.cpp b/Engine/Plugins/AI/MLAdapter/Source/MLAdapter/Private/Sensors/MLAdapterSensor_AIPerception.cpp
--- a/Engine/Plugins/AI/MLAdapter/Source/MLAdapter/Private/Sensors/MLAdapterSensor_AIPerception.cpp
+++ b/Engine/Plugins/AI/MLAdapter/Source/MLAdapter/Private/Sensors/MLAdapterSensor_AIPerception.cpp
@@ -217,11 +217,14 @@ void UMLAdapterSensor_AIPerception::GetObservations(FMLAdapterMemoryWriter& Ar)
 		Ar.Serialize(&TargetData.Distance, sizeof(float));
 		if (bVectorMode)
 		{
-			Ar << TargetData.HeadingVector;
+			FVector3f HeadingVector3f = TargetData.HeadingVector;
+			Ar << HeadingVector3f;
 		}
 		else
 		{
-			Ar << TargetData.HeadingRotator.Pitch << TargetData.HeadingRotator.Yaw;
+			float Pitch = TargetData.HeadingRotator.Pitch;
+			float Yaw = TargetData.HeadingRotator.Yaw;
+			Ar << Pitch << Yaw;
 		}
 	}
 }
diff --git a/Engine/Plugins/AI/MLAdapter/Source/MLAdapter/Private/Sensors/MLAdapterSensor_Movement.cpp b/Engine/Plugins/AI/MLAdapter/Source/MLAdapter/Private/Sensors/MLAdapterSensor_Movement.cpp
--- a/Engine/Plugins/AI/MLAdapter/Source/MLAdapter/Private/Sensors/MLAdapterSensor_Movement.cpp
+++ b/Engine/Plugins/AI/MLAdapter/Source/MLAdapter/Private/Sensors/MLAdapterSensor_Movement.cpp
@@ -74,8 +74,8 @@ void UMLAdapterSensor_Movement::OnAvatarSet(AActor* Avatar)
 void UMLAdapterSensor_Movement::GetObservations(FMLAdapterMemoryWriter& Ar)
 {
 	FScopeLock Lock(&ObservationCS);
-	FVector Location = bAbsoluteLocation ? CurrentLocation : (CurrentLocation - RefLocation);
-	FVector Velocity = bAbsoluteVelocity ? CurrentVelocity : (CurrentVelocity - RefVelocity);
+	FVector3f Location = bAbsoluteLocation ? CurrentLocation : (CurrentLocation - RefLocation);
+	FVector3f Velocity = bAbsoluteVelocity ? CurrentVelocity : (CurrentVelocity - RefVelocity);
 	
 	FMLAdapter::FSpaceSerializeGuard SerializeGuard(SpaceDef, Ar);	
 	Ar << Location << Velocity;
