MLAdapter: Add debug option and fixed_time_step option to runner.py

#rnx

diff --git a/Engine/Plugins/AI/MLAdapter/Source/python/unreal/mladapter/runner.py b/Engine/Plugins/AI/MLAdapter/Source/python/unreal/mladapter/runner.py
--- a/Engine/Plugins/AI/MLAdapter/Source/python/unreal/mladapter/runner.py
+++ b/Engine/Plugins/AI/MLAdapter/Source/python/unreal/mladapter/runner.py
@@ -18,12 +18,12 @@ _EXECUTABLE_OVERRIDE = None
 
 
 
-def set_debug(debug):
+def set_debug(debug, config=''):
     global EXEC_NAME
     global __DEBUG
     __DEBUG = debug
     if __DEBUG:
-        EXEC_NAME = 'UnrealEditor-Win64-Debug.exe'
+        EXEC_NAME = f'UnrealEditor-Win64-Debug{config}.exe'
     else:
         EXEC_NAME = 'UnrealEditor.exe'
 
@@ -58,11 +58,12 @@ class UEParams(object):
     GAME = 'game'
     UNATTENDED = 'unattended'
     
-    def __init__(self, map_name='', rendering=True, sound=True, single_thread=False, custom=''):
-        self.__options = [UEParams.WINDOW, UEParams.FIXED_TIME_STEP, UEParams.GAME, UEParams.UNATTENDED]
+    def __init__(self, map_name='', rendering=True, sound=True, single_thread=False, fixed_time_step=False, custom=''):
+        self.__options = [UEParams.WINDOW, UEParams.GAME, UEParams.UNATTENDED]
         self.__params = {'resx': 320, 'resy': 240, 'fps': DEFAULT_FPS}
         self.map_name = map_name
         self.custom = custom
+        self.fixed_time_step(fixed_time_step)
         self.rendering(rendering)
         self.sound(sound)
         self.single_thread(single_thread)
@@ -92,6 +93,9 @@ class UEParams(object):
         self.set_enable_option('onethread', enabled)
         self.set_enable_option('ReduceThreadUsage', enabled)
 
+    def fixed_time_step(self, enabled):
+        self.set_enable_option(UEParams.FIXED_TIME_STEP, enabled)
+
     def sound(self, enabled):
         self.set_enable_option('nosound', not enabled)
 
