MLAdapter: New Action Duration Option & Related Changes

#rnx

diff --git a/Engine/Plugins/AI/MLAdapter/Source/MLAdapter/Private/Agents/MLAdapterAgent.cpp b/Engine/Plugins/AI/MLAdapter/Source/MLAdapter/Private/Agents/MLAdapterAgent.cpp
--- a/Engine/Plugins/AI/MLAdapter/Source/MLAdapter/Private/Agents/MLAdapterAgent.cpp
+++ b/Engine/Plugins/AI/MLAdapter/Source/MLAdapter/Private/Agents/MLAdapterAgent.cpp
@@ -171,10 +171,21 @@ void UMLAdapterAgent::OnPawnControllerChanged(APawn* InPawn, AController* InCont
 
 void UMLAdapterAgent::Sense(const float DeltaTime)
 {
-	for (UMLAdapterSensor* Sensor : Sensors)
+	CurrentActionDuration += DeltaTime;
+
+	bool ShouldDecideNow = !bEnableActionDuration || CurrentActionDuration >= ActionDurationSeconds;
+	if (ShouldDecideNow)
 	{
-		// not that due to the system's design Sensor won't be null
-		Sensor->Sense(DeltaTime);
+		FScopeLock Lock(&ActionDurationCS);
+
+		CurrentActionDuration = 0.f;
+		bActionDurationElapsed = true;
+
+		for (UMLAdapterSensor* Sensor : Sensors)
+		{
+			// not that due to the system's design Sensor won't be null
+			Sensor->Sense(DeltaTime);
+		}
 	}
 }
 
@@ -444,6 +455,25 @@ bool UMLAdapterAgent::IsDone() const
 		&& bEverHadAvatar == true;
 }
 
+void UMLAdapterAgent::EnableActionDuration(bool bEnable, float DurationSeconds)
+{
+	bEnableActionDuration = bEnable;
+	ActionDurationSeconds = DurationSeconds;
+}
+
+bool UMLAdapterAgent::TryResetActionDuration()
+{
+	FScopeLock Lock(&ActionDurationCS);
+
+	if (bActionDurationElapsed)
+	{
+		bActionDurationElapsed = false;
+		return true;
+	}
+
+	return false;
+}
+
 #if WITH_GAMEPLAY_DEBUGGER
 #include "GameplayDebuggerCategory.h"
 
diff --git a/Engine/Plugins/AI/MLAdapter/Source/MLAdapter/Private/Agents/MLAdapterAgent_Inference.cpp b/Engine/Plugins/AI/MLAdapter/Source/MLAdapter/Private/Agents/MLAdapterAgent_Inference.cpp
--- a/Engine/Plugins/AI/MLAdapter/Source/MLAdapter/Private/Agents/MLAdapterAgent_Inference.cpp
+++ b/Engine/Plugins/AI/MLAdapter/Source/MLAdapter/Private/Agents/MLAdapterAgent_Inference.cpp
@@ -24,6 +24,14 @@ void UMLAdapterAgent_Inference::Think(const float DeltaTime)
 		return;
 	}
 
+	if (bEnableActionDuration && !bActionDurationElapsed)
+	{
+		// Action duration is enabled and it's not time to think yet.
+		return;
+	}
+
+	bActionDurationElapsed = false;
+
 	TArray<uint8> Buffer;
 	FMLAdapterMemoryWriter Writer(Buffer);
 	GetObservations(Writer);
diff --git a/Engine/Plugins/AI/MLAdapter/Source/MLAdapter/Private/Managers/MLAdapterManager_Server.cpp b/Engine/Plugins/AI/MLAdapter/Source/MLAdapter/Private/Managers/MLAdapterManager_Server.cpp
--- a/Engine/Plugins/AI/MLAdapter/Source/MLAdapter/Private/Managers/MLAdapterManager_Server.cpp
+++ b/Engine/Plugins/AI/MLAdapter/Source/MLAdapter/Private/Managers/MLAdapterManager_Server.cpp
@@ -30,7 +30,7 @@ void UMLAdapterManager::ConfigureAsServer(FRPCServer& Server)
 	Server.bind("enable_manual_world_tick", [this](bool bEnable) {
 		SetManualWorldTickEnabled(bEnable);
 	});
-	Librarian.AddRPCFunctionDescription(TEXT("enable_manual_world_tick"), TEXT("(), Controlls whether the world is running real time or it\'s being ticked manually with calls to \'step\' or \'request_world_tick\' functions. Default is \'real time\'."));
+	Librarian.AddRPCFunctionDescription(TEXT("enable_manual_world_tick"), TEXT("(bool bEnable), Controls whether the world is running real time or it\'s being ticked manually with calls to \'step\' or \'request_world_tick\' functions. Default is \'real time\'."));
 
 	Server.bind("request_world_tick", [this](int32 TickCount, bool bWaitForWorldTick) {
 		if (bTickWorldManually == false)
@@ -46,6 +46,24 @@ void UMLAdapterManager::ConfigureAsServer(FRPCServer& Server)
 	});
 	Librarian.AddRPCFunctionDescription(TEXT("request_world_tick"), TEXT("(int TickCount, bool bWaitForWorldTick), Requests a TickCount world ticks. This has meaning only if \'enable_manual_world_tick(true)\' has been called prior to this function. If bWaitForWorldTick is true then the call will not return until the world has been ticked required number of times"));
 
+	Server.bind("enable_action_duration", [this](FMLAdapter::FAgentID AgentID, bool bEnableActionDuration, float DurationSeconds) {
+		if (Session)
+		{
+			Session->EnableActionDuration(AgentID, bEnableActionDuration, DurationSeconds);
+		}
+	});
+	Librarian.AddRPCFunctionDescription(TEXT("enable_action_duration"), TEXT("(uint AgentID, bool bEnableActionDuration, float DurationSeconds), Enable/disable the action durations on the agent with the specified time duration in seconds."));
+
+	Server.bind("wait_for_action_duration", [this](FMLAdapter::FAgentID AgentID) {
+		if (Session)
+		{
+			while (!Session->TryResetActionDuration(AgentID))
+			{
+			}
+		}
+	});
+	Librarian.AddRPCFunctionDescription(TEXT("wait_for_action_duration"), TEXT("(uint AgentID), Wait for the action duration to elapse for the agent. Only works if \'enable_action_duration\' has been called previously."));
+
 	Server.bind("close_session", [this]() { UMLAdapterManager::Get().SetSession(nullptr); });
 	Librarian.AddRPCFunctionDescription(TEXT("close_session"), TEXT("(), shuts down the current session (along with all the agents)."));
 #endif // WITH_RPCLIB
diff --git a/Engine/Plugins/AI/MLAdapter/Source/MLAdapter/Private/Sessions/MLAdapterSession.cpp b/Engine/Plugins/AI/MLAdapter/Source/MLAdapter/Private/Sessions/MLAdapterSession.cpp
--- a/Engine/Plugins/AI/MLAdapter/Source/MLAdapter/Private/Sessions/MLAdapterSession.cpp
+++ b/Engine/Plugins/AI/MLAdapter/Source/MLAdapter/Private/Sessions/MLAdapterSession.cpp
@@ -42,6 +42,26 @@ void UMLAdapterSession::SetManualWorldTickEnabled(bool bEnable)
 	}
 }
 
+void UMLAdapterSession::EnableActionDuration(FMLAdapter::FAgentID AgentID, bool bEnable, float DurationSeconds)
+{
+	UMLAdapterAgent* Agent = GetAgent(AgentID);
+	if (Agent != nullptr)
+	{
+		Agent->EnableActionDuration(bEnable, DurationSeconds);		
+	}
+}
+
+bool UMLAdapterSession::TryResetActionDuration(FMLAdapter::FAgentID AgentID)
+{
+	UMLAdapterAgent* Agent = GetAgent(AgentID);
+	if (Agent != nullptr)
+	{
+		return Agent->TryResetActionDuration();
+	}
+
+	return false;
+}
+
 void UMLAdapterSession::FWorldTicker::Tick(float DeltaTime) 
 {
 #if WITH_EDITORONLY_DATA
@@ -238,7 +258,7 @@ void UMLAdapterSession::Close()
 void UMLAdapterSession::Tick(float DeltaTime)
 {
 	LastTimestamp = CachedWorld ? CachedWorld->GetTimeSeconds() : -1.f;
-
+	
 	// @todo for perf reasons we could grab all the agents' senses and tick them
 	// by class to keep the cache hot
 	for (UMLAdapterAgent* Agent : Agents)
diff --git a/Engine/Plugins/AI/MLAdapter/Source/MLAdapter/Public/Agents/MLAdapterAgent.h b/Engine/Plugins/AI/MLAdapter/Source/MLAdapter/Public/Agents/MLAdapterAgent.h
--- a/Engine/Plugins/AI/MLAdapter/Source/MLAdapter/Public/Agents/MLAdapterAgent.h
+++ b/Engine/Plugins/AI/MLAdapter/Source/MLAdapter/Public/Agents/MLAdapterAgent.h
@@ -184,6 +184,12 @@ public:
 	/** Returns true if this agent has an avatar set. */
 	bool IsReady() const { return Avatar != nullptr; }
 
+	/** Enable/disable the action durations with the specified time duration in seconds. */
+	void EnableActionDuration(bool bEnable, float DurationSeconds);
+
+	/** Resets the action duration flag if it has elapsed. Returns false if not reset yet. Used with HasActionDurationElapsed by Manager. */
+	bool TryResetActionDuration();
+
 protected:
 	virtual void ShutDownSensorsAndActuators();
 
@@ -197,6 +203,23 @@ protected:
 	UPROPERTY(EditDefaultsOnly, BlueprintReadOnly, Category = MLAdapter)
 	TSubclassOf<AActor> AvatarClass;
 
+	// If true, then agents won't be able to make a new decision until each action duration has elapsed.
+	UPROPERTY(EditAnywhere, Category = MLAdapter)
+	bool bEnableActionDuration = false;
+
+	// How long should agents wait before they can change their action
+	UPROPERTY(EditAnywhere, Category = MLAdapter)
+	float ActionDurationSeconds = 0.1f;
+
+	// How much time has the current action been executing
+	UPROPERTY(VisibleInstanceOnly, Category = MLAdapter)
+	float CurrentActionDuration = 0.f;
+
+	UPROPERTY(VisibleInstanceOnly, Category = MLAdapter)
+	bool bActionDurationElapsed = false;
+
+	mutable FCriticalSection ActionDurationCS;
+
 private:
 	UPROPERTY()
 	AActor* Avatar;
diff --git a/Engine/Plugins/AI/MLAdapter/Source/MLAdapter/Public/Sessions/MLAdapterSession.h b/Engine/Plugins/AI/MLAdapter/Source/MLAdapter/Public/Sessions/MLAdapterSession.h
--- a/Engine/Plugins/AI/MLAdapter/Source/MLAdapter/Public/Sessions/MLAdapterSession.h
+++ b/Engine/Plugins/AI/MLAdapter/Source/MLAdapter/Public/Sessions/MLAdapterSession.h
@@ -102,6 +102,12 @@ public:
 	/** If true, the world's ticking will be controlled by the remote client. */
 	void SetManualWorldTickEnabled(bool bEnable);
 
+	/** Enable/disable the action durations with the specified time duration in seconds. */
+	void EnableActionDuration(FMLAdapter::FAgentID AgentID, bool bEnable, float DurationSeconds);
+
+	/** Resets the action duration flag if it has elapsed. Returns false if not reset yet or the agent is not found. Used with HasActionDurationElapsed by Manager. */
+	bool TryResetActionDuration(FMLAdapter::FAgentID AgentID);
+
 	FOnAgentAvatarChangedDelegate GetOnAgentAvatarChanged() { return OnAgentAvatarChanged; }
 	FOnBeginAgentRemove GetOnBeginAgentRemove() { return OnBeginAgentRemove; }
 
diff --git a/Engine/Plugins/AI/MLAdapter/Source/python/unreal/mladapter/core.py b/Engine/Plugins/AI/MLAdapter/Source/python/unreal/mladapter/core.py
--- a/Engine/Plugins/AI/MLAdapter/Source/python/unreal/mladapter/core.py
+++ b/Engine/Plugins/AI/MLAdapter/Source/python/unreal/mladapter/core.py
@@ -54,12 +54,13 @@ class UnrealEnv(gym.Env):
     PROJECT_NAME = 'GenericUnrealEnvironment'
     
     def __init__(self, server_address=LOCALHOST, server_port=DEFAULT_PORT, agent_config=None, reacquire=True, 
-                 realtime=False, auto_connect=True, timeout=DEFAULT_TIMEOUT, ue_params=None, project_name=None):
+                 realtime=False, auto_connect=True, timeout=DEFAULT_TIMEOUT, ue_params=None, project_name=None,
+                 use_preconfigured_agent=False, action_duration_seconds=None):
 
-        # it's expected that the used passes in project_name or implements a wrapper defining PROJECT_NAME
+        # It's expected that the user passes in project_name or implements a wrapper defining PROJECT_NAME
         self._project_name = project_name or self.__class__.PROJECT_NAME
 
-        if agent_config is None:
+        if not use_preconfigured_agent and agent_config is None:
             agent_config = self.__class__.default_agent_config()
 
         self.__agent_config = agent_config
@@ -68,6 +69,7 @@ class UnrealEnv(gym.Env):
         self._debug_id = server_port
         self.__server_port = server_port
         self._realtime = None
+        self._action_duration_seconds = action_duration_seconds
 
         self.__engine_process = None
         if ue_params is not None:
@@ -143,16 +145,29 @@ class UnrealEnv(gym.Env):
 
     def _set_realtime(self, realtime):
         self._realtime = realtime
-        self._world_step_impl = self._tick_world if not realtime else lambda: None
+        
+        if not realtime:
+            self._world_step_impl = self._tick_world
+        elif self._action_duration_seconds != None:
+            self._world_step_impl = self.wait_for_action_duration
+        else:
+            self._world_step_impl = lambda: None
+
         if hasattr(self.__rpc_client, 'enable_manual_world_tick'):
             self.__rpc_client.enable_manual_world_tick(not realtime)
 
+        if hasattr(self.__rpc_client, 'enable_action_duration'):
+            self.__rpc_client.enable_action_duration(self.__agent_id, True, self._action_duration_seconds)
+
     def _get_observation(self):
         raw_obs = self.__rpc_client.get_observations(self.__agent_id)
         return gym.spaces.unflatten(self.observation_space, raw_obs)
 
     def _tick_world(self):
         self.__rpc_client.request_world_tick(self._frames_step, True)
+    
+    def wait_for_action_duration(self):
+        self.__rpc_client.wait_for_action_duration(self.__agent_id)
 
     def reset(self, wait_action=None, skip_time=1):
         if not self.is_connected():
@@ -228,7 +243,7 @@ class UnrealEnv(gym.Env):
             action = self.wrap_action(action)
             self.__rpc_client.act(self.__agent_id, action)
         
-        self._steps_performed += 1        
+        self._steps_performed += 1
         self._world_step_impl()
         observation = self._get_observation()
         reward = self.get_reward()
