MLAdapter: Enhanced Input Actuator - Clear Action On Use

#rnx

diff --git a/Engine/Plugins/AI/MLAdapter/Source/MLAdapter/Private/Actuators/MLAdapterActuator_EnhancedInput.cpp b/Engine/Plugins/AI/MLAdapter/Source/MLAdapter/Private/Actuators/MLAdapterActuator_EnhancedInput.cpp
--- a/Engine/Plugins/AI/MLAdapter/Source/MLAdapter/Private/Actuators/MLAdapterActuator_EnhancedInput.cpp
+++ b/Engine/Plugins/AI/MLAdapter/Source/MLAdapter/Private/Actuators/MLAdapterActuator_EnhancedInput.cpp
@@ -17,6 +17,13 @@ void UMLAdapterActuator_EnhancedInput::Configure(const TMap<FName, FString>& Par
 {
 	Super::Configure(Params);
 
+	const FName NAME_bClearActionOnUse = TEXT("bClearActionOnUse");
+	const FString* bClearActionOnUseValue = Params.Find(NAME_bClearActionOnUse);
+	if (bClearActionOnUseValue != nullptr)
+	{
+		bClearActionOnUse = (*bClearActionOnUseValue != "0");
+	}
+
 	UpdateSpaceDef();
 }
 
@@ -93,15 +100,31 @@ void UMLAdapterActuator_EnhancedInput::Act(const float DeltaTime)
 		}
 	}
 	
-	InputData.Empty(InputData.Num());
+	if (bClearActionOnUse)
+	{
+		InputData.Empty(InputData.Num());
+	}
 }
 
 void UMLAdapterActuator_EnhancedInput::DigestInputData(FMLAdapterMemoryReader& ValueStream)
 {
 	FScopeLock Lock(&ActionCS);
 
-	const int32 OldSize = InputData.Num();
-	InputData.AddUninitialized(SpaceDef->Num());
-	// offsetting the serialization since there might be unprocessed data in InputData
-	ValueStream.Serialize(InputData.GetData() + OldSize, SpaceDef->Num() * sizeof(float));
+	if (bClearActionOnUse)
+	{
+		const int32 OldSize = InputData.Num();
+		InputData.AddUninitialized(SpaceDef->Num());
+		// offsetting the serialization since there might be unprocessed data in InputData
+		ValueStream.Serialize(InputData.GetData() + OldSize, SpaceDef->Num() * sizeof(float));
+	}
+	else
+	{
+		if (InputData.Num() == 0)
+		{
+			InputData.AddUninitialized(SpaceDef->Num());
+		}
+
+		// Overwrite the currently executing actions
+		ValueStream.Serialize(InputData.GetData(), SpaceDef->Num() * sizeof(float));
+	}
 }
\ No newline at end of file
diff --git a/Engine/Plugins/AI/MLAdapter/Source/MLAdapter/Public/Actuators/MLAdapterActuator_EnhancedInput.h b/Engine/Plugins/AI/MLAdapter/Source/MLAdapter/Public/Actuators/MLAdapterActuator_EnhancedInput.h
--- a/Engine/Plugins/AI/MLAdapter/Source/MLAdapter/Public/Actuators/MLAdapterActuator_EnhancedInput.h
+++ b/Engine/Plugins/AI/MLAdapter/Source/MLAdapter/Public/Actuators/MLAdapterActuator_EnhancedInput.h
@@ -36,4 +36,7 @@ protected:
 	*/
 	UPROPERTY(EditDefaultsOnly, BlueprintReadOnly, Category = MLAdapter)
 	TArray<TObjectPtr<UInputAction>> TrackedActions;
+
+	UPROPERTY(EditDefaultsOnly, BlueprintReadOnly, Category = MLAdapter)
+	bool bClearActionOnUse = true;
 };
\ No newline at end of file
