MLAdapter: Convert param looping to Find calls

#rnx

diff --git a/Engine/Plugins/AI/MLAdapter/Source/MLAdapter/Private/Actuators/MLAdapterActuator_Camera.cpp b/Engine/Plugins/AI/MLAdapter/Source/MLAdapter/Private/Actuators/MLAdapterActuator_Camera.cpp
--- a/Engine/Plugins/AI/MLAdapter/Source/MLAdapter/Private/Actuators/MLAdapterActuator_Camera.cpp
+++ b/Engine/Plugins/AI/MLAdapter/Source/MLAdapter/Private/Actuators/MLAdapterActuator_Camera.cpp
@@ -27,16 +27,13 @@ UMLAdapterActuator_Camera::UMLAdapterActuator_Camera(const FObjectInitializer& O
 
 void UMLAdapterActuator_Camera::Configure(const TMap<FName, FString>& Params)
 {
-	const FName NAME_Mode = TEXT("mode");
-
 	Super::Configure(Params);
 
-	for (auto KeyValue : Params)
+	const FName NAME_Mode = TEXT("mode");
+	const FString* ModeValue = Params.Find(NAME_Mode);
+	if (ModeValue != nullptr)
 	{
-		if (KeyValue.Key == NAME_Mode)
-		{
-			bVectorMode = (KeyValue.Value.Find(TEXT("vector")) != INDEX_NONE);
-		}
+		bVectorMode = (ModeValue->Find(TEXT("vector")) != INDEX_NONE);
 	}
 
 	UpdateSpaceDef();
diff --git a/Engine/Plugins/AI/MLAdapter/Source/MLAdapter/Private/Actuators/MLAdapterActuator_InputKey.cpp b/Engine/Plugins/AI/MLAdapter/Source/MLAdapter/Private/Actuators/MLAdapterActuator_InputKey.cpp
--- a/Engine/Plugins/AI/MLAdapter/Source/MLAdapter/Private/Actuators/MLAdapterActuator_InputKey.cpp
+++ b/Engine/Plugins/AI/MLAdapter/Source/MLAdapter/Private/Actuators/MLAdapterActuator_InputKey.cpp
@@ -16,37 +16,35 @@ UMLAdapterActuator_InputKey::UMLAdapterActuator_InputKey(const FObjectInitialize
 
 void UMLAdapterActuator_InputKey::Configure(const TMap<FName, FString>& Params) 
 {
-	const FName NAME_IgnoreKeys = TEXT("ignore_keys");
-	const FName NAME_IgnoreActions = TEXT("ignore_actions");
-
 	Super::Configure(Params);
 
 	TArray<FName> IgnoreKeys;
-	TArray<FName> IgnoreActions;
-
-	for (auto KeyValue : Params)
+	const FName NAME_IgnoreKeys = TEXT("ignore_keys");
+	const FString* IgnoreKeysValue = Params.Find(NAME_IgnoreKeys);
+	if (IgnoreKeysValue != nullptr)
 	{
-		if (KeyValue.Key == NAME_IgnoreKeys)
+		TArray<FString> Tokens;
+		// contains a list of keys to not press
+		IgnoreKeysValue->ParseIntoArrayWS(Tokens, TEXT(","));
+		IgnoreKeys.Reserve(Tokens.Num());
+		for (const FString& Token : Tokens)
 		{
-			TArray<FString> Tokens;
-			// contains a list of keys to not press
-			KeyValue.Value.ParseIntoArrayWS(Tokens, TEXT(","));
-			IgnoreKeys.Reserve(Tokens.Num());
-			for (const FString& Token : Tokens)
-			{
-				IgnoreKeys.Add(FName(Token));
-			}
+			IgnoreKeys.Add(FName(Token));
 		}
-		else if (KeyValue.Key == NAME_IgnoreActions)
+	}
+
+	TArray<FName> IgnoreActions;
+	const FName NAME_IgnoreActions = TEXT("ignore_actions");
+	const FString* IgnoreActionsValue = Params.Find(NAME_IgnoreActions);
+	if (IgnoreActionsValue != nullptr)
+	{
+		TArray<FString> Tokens;
+		// contains a list of actions to ignore
+		IgnoreActionsValue->ParseIntoArrayWS(Tokens, TEXT(","));
+		IgnoreActions.Reserve(Tokens.Num());
+		for (const FString& Token : Tokens)
 		{
-			TArray<FString> Tokens;
-			// contains a list of keys to not press
-			KeyValue.Value.ParseIntoArrayWS(Tokens, TEXT(","));
-			IgnoreActions.Reserve(Tokens.Num());
-			for (const FString& Token : Tokens)
-			{
-				IgnoreActions.Add(FName(Token));
-			}
+			IgnoreActions.Add(FName(Token));
 		}
 	}
 
diff --git a/Engine/Plugins/AI/MLAdapter/Source/MLAdapter/Private/Sensors/MLAdapterSensor_AIPerception.cpp b/Engine/Plugins/AI/MLAdapter/Source/MLAdapter/Private/Sensors/MLAdapterSensor_AIPerception.cpp
--- a/Engine/Plugins/AI/MLAdapter/Source/MLAdapter/Private/Sensors/MLAdapterSensor_AIPerception.cpp
+++ b/Engine/Plugins/AI/MLAdapter/Source/MLAdapter/Private/Sensors/MLAdapterSensor_AIPerception.cpp
@@ -25,36 +25,41 @@ UMLAdapterSensor_AIPerception::UMLAdapterSensor_AIPerception(const FObjectInitia
 
 void UMLAdapterSensor_AIPerception::Configure(const TMap<FName, FString>& Params)
 {
+	Super::Configure(Params);
+
 	const FName NAME_Count = TEXT("count");
+	const FString* CountValue = Params.Find(NAME_Count);
+	if (CountValue != nullptr)
+	{
+		TargetsToSenseCount = FMath::Max(FCString::Atoi((TCHAR*)CountValue), 1);
+	}
+
 	const FName NAME_Sort = TEXT("sort");
+	const FString* SortValue = Params.Find(NAME_Sort);
+	if (SortValue != nullptr)
+	{
+		TargetsSortType = (*SortValue == TEXT("in_front")) ? ESortType::InFrontness : ESortType::Distance;
+	}
+
 	const FName NAME_Mode = TEXT("mode");
-	const FName NAME_PeripheralAngle = TEXT("peripheral_angle");
-	const FName NAME_MaxAge = TEXT("max_age");
+	const FString* ModeValue = Params.Find(NAME_Mode);
+	if (ModeValue != nullptr)
+	{
+		bVectorMode = (ModeValue->Find(TEXT("vector")) != INDEX_NONE);
+	}
 
-	Super::Configure(Params);
+	const FName NAME_PeripheralAngle = TEXT("peripheral_angle");
+	const FString* PeripheralAngleValue = Params.Find(NAME_PeripheralAngle);
+	if (PeripheralAngleValue != nullptr)
+	{
+		PeripheralVisionAngleDegrees = FMath::Max(FCString::Atof((TCHAR*)PeripheralAngleValue), 1.f);
+	}
 
-	for (auto KeyValue : Params)
+	const FName NAME_MaxAge = TEXT("max_age");
+	const FString* MaxAgeValue = Params.Find(NAME_MaxAge);
+	if (MaxAgeValue != nullptr)
 	{
-		if (KeyValue.Key == NAME_Count)
-		{
-			TargetsToSenseCount = FMath::Max(FCString::Atoi(*KeyValue.Value), 1);
-		}
-		else if (KeyValue.Key == NAME_Sort)
-		{
-			TargetsSortType = (KeyValue.Value == TEXT("in_front")) ? ESortType::InFrontness : ESortType::Distance;
-		}
-		else if (KeyValue.Key == NAME_Mode)
-		{
-			bVectorMode = (KeyValue.Value.Find(TEXT("vector")) != INDEX_NONE);
-		}
-		else if (KeyValue.Key == NAME_PeripheralAngle)
-		{
-			PeripheralVisionAngleDegrees = FMath::Max(FCString::Atof(*KeyValue.Value), 1.f);
-		}
-		else if (KeyValue.Key == NAME_MaxAge)
-		{
-			MaxStimulusAge = FMath::Max(FCString::Atof(*KeyValue.Value), 0.001f);
-		}
+		MaxStimulusAge = FMath::Max(FCString::Atof((TCHAR*)MaxAgeValue), 0.001f);
 	}
 
 	UpdateSpaceDef();
diff --git a/Engine/Plugins/AI/MLAdapter/Source/MLAdapter/Private/Sensors/MLAdapterSensor_Attribute.cpp b/Engine/Plugins/AI/MLAdapter/Source/MLAdapter/Private/Sensors/MLAdapterSensor_Attribute.cpp
--- a/Engine/Plugins/AI/MLAdapter/Source/MLAdapter/Private/Sensors/MLAdapterSensor_Attribute.cpp
+++ b/Engine/Plugins/AI/MLAdapter/Source/MLAdapter/Private/Sensors/MLAdapterSensor_Attribute.cpp
@@ -17,18 +17,15 @@ bool UMLAdapterSensor_Attribute::ConfigureForAgent(UMLAdapterAgent& Agent)
 
 void UMLAdapterSensor_Attribute::Configure(const TMap<FName, FString>& Params)
 {
-	const FName NAME_Attributes = TEXT("attributes");
-	
 	Super::Configure(Params);
 
-	for (auto KeyValue : Params)
+	const FName NAME_Attributes = TEXT("attributes");
+	const FString* AttributesValue = Params.Find(NAME_Attributes);
+	if (AttributesValue != nullptr)
 	{
-		if (KeyValue.Key == NAME_Attributes)
-		{
-			TArray<FString> Tokens;
-			KeyValue.Value.ParseIntoArrayWS(Tokens, TEXT(","));
-			SetAttributes(Tokens);
-		}
+		TArray<FString> Tokens;
+		AttributesValue->ParseIntoArrayWS(Tokens, TEXT(","));
+		SetAttributes(Tokens);
 	}
 }
 
diff --git a/Engine/Plugins/AI/MLAdapter/Source/MLAdapter/Private/Sensors/MLAdapterSensor_Input.cpp b/Engine/Plugins/AI/MLAdapter/Source/MLAdapter/Private/Sensors/MLAdapterSensor_Input.cpp
--- a/Engine/Plugins/AI/MLAdapter/Source/MLAdapter/Private/Sensors/MLAdapterSensor_Input.cpp
+++ b/Engine/Plugins/AI/MLAdapter/Source/MLAdapter/Private/Sensors/MLAdapterSensor_Input.cpp
@@ -18,18 +18,15 @@ UMLAdapterSensor_Input::UMLAdapterSensor_Input(const FObjectInitializer& ObjectI
 
 void UMLAdapterSensor_Input::Configure(const TMap<FName, FString>& Params)
 {
-	const FName NAME_RecordRelease = TEXT("record_release");
-
 	Super::Configure(Params);
 
-	for (auto KeyValue : Params)
+	const FName NAME_RecordRelease = TEXT("record_release");
+	const FString* RecordReleaseValue = Params.Find(NAME_RecordRelease);
+	if (RecordReleaseValue != nullptr)
 	{
-		if (KeyValue.Key == NAME_RecordRelease)
-		{
-			bool bValue = bRecordKeyRelease;
-			LexFromString(bValue, *KeyValue.Value);
-			bRecordKeyRelease = bValue;
-		}
+		bool bValue = bRecordKeyRelease;
+		LexFromString(bValue, (TCHAR*)RecordReleaseValue);
+		bRecordKeyRelease = bValue;
 	}
 
 	FMLAdapterInputHelper::CreateInputMap(InterfaceKeys, FKeyToInterfaceKeyMap);
diff --git a/Engine/Plugins/AI/MLAdapter/Source/MLAdapter/Private/Sensors/MLAdapterSensor_Movement.cpp b/Engine/Plugins/AI/MLAdapter/Source/MLAdapter/Private/Sensors/MLAdapterSensor_Movement.cpp
--- a/Engine/Plugins/AI/MLAdapter/Source/MLAdapter/Private/Sensors/MLAdapterSensor_Movement.cpp
+++ b/Engine/Plugins/AI/MLAdapter/Source/MLAdapter/Private/Sensors/MLAdapterSensor_Movement.cpp
@@ -22,21 +22,20 @@ bool UMLAdapterSensor_Movement::ConfigureForAgent(UMLAdapterAgent& Agent)
 
 void UMLAdapterSensor_Movement::Configure(const TMap<FName, FString>& Params)
 {
-	const FName LocationKeyName = TEXT("location");
-	const FName VelocityKeyName = TEXT("velocity");
-	
 	Super::Configure(Params);
 
-	for (auto KeyValue : Params)
+	const FName LocationKeyName = TEXT("location");
+	const FString* LocationKeyValue = Params.Find(LocationKeyName);
+	if (LocationKeyValue != nullptr)
+	{
+		bAbsoluteLocation = (LocationKeyValue->Find(TEXT("absolute")) != INDEX_NONE);
+	}
+
+	const FName VelocityKeyName = TEXT("velocity");
+	const FString* VelocityKeyValue = Params.Find(VelocityKeyName);
+	if (VelocityKeyValue != nullptr)
 	{
-		if (KeyValue.Key == LocationKeyName)
-		{
-			bAbsoluteLocation = (KeyValue.Value.Find(TEXT("absolute")) != INDEX_NONE);
-		}
-		else if (KeyValue.Key == VelocityKeyName)
-		{
-			bAbsoluteLocation = (KeyValue.Value.Find(TEXT("absolute")) != INDEX_NONE);
-		}
+		bAbsoluteVelocity = (VelocityKeyValue->Find(TEXT("absolute")) != INDEX_NONE);
 	}
 
 	UpdateSpaceDef();
