MLAdapter: Fix passing none in Python code

#rnx

diff --git a/Engine/Plugins/AI/MLAdapter/Source/python/unreal/mladapter/core.py b/Engine/Plugins/AI/MLAdapter/Source/python/unreal/mladapter/core.py
--- a/Engine/Plugins/AI/MLAdapter/Source/python/unreal/mladapter/core.py
+++ b/Engine/Plugins/AI/MLAdapter/Source/python/unreal/mladapter/core.py
@@ -156,7 +156,7 @@ class UnrealEnv(gym.Env):
         if hasattr(self.__rpc_client, 'enable_manual_world_tick'):
             self.__rpc_client.enable_manual_world_tick(not realtime)
 
-        if hasattr(self.__rpc_client, 'enable_action_duration'):
+        if hasattr(self.__rpc_client, 'enable_action_duration') and self._action_duration_seconds != None:
             self.__rpc_client.enable_action_duration(self.__agent_id, True, self._action_duration_seconds)
 
     def _get_observation(self):
