MLAdapter: simplify world filtering in PIE

#rnx

diff --git a/Engine/Plugins/AI/MLAdapter/Source/MLAdapter/Private/Managers/MLAdapterManager.cpp b/Engine/Plugins/AI/MLAdapter/Source/MLAdapter/Private/Managers/MLAdapterManager.cpp
--- a/Engine/Plugins/AI/MLAdapter/Source/MLAdapter/Private/Managers/MLAdapterManager.cpp
+++ b/Engine/Plugins/AI/MLAdapter/Source/MLAdapter/Private/Managers/MLAdapterManager.cpp
@@ -374,35 +374,11 @@ bool UMLAdapterManager::ShouldInitForWorld(const UWorld& World) const
 	}
 
 #if WITH_EDITOR
-	// PIE is a special case, we need to see if it's a client-server PIE and 
-	// if so we need to filter the incoming world based on the settings
-	if (World.WorldType == EWorldType::PIE)
-	{
-		if (World.HasAnyFlags(RF_WasLoaded))
-		{
-			const ULevelEditorPlaySettings* PlayInSettings = GetDefault<ULevelEditorPlaySettings>();
-			if (PlayInSettings)
-			{
-				EPlayNetMode PlayNetMode = PIE_Standalone;
-				PlayInSettings->GetPlayNetMode(PlayNetMode);
-
-				if (PlayNetMode != PIE_Standalone)
-				{
-					UE_LOG(LogUnrealEditorMLAdapter, Log, TEXT("Ignoring %s due to net mode != PIE_Standalone"), *World.GetName());
-					return false;
-				}
-			}
-			else
-			{
-				// skipping temp world
-				return false;
-			}
-		}
-
-	}
-#endif // WITH_EDITOR
-
+	return World.WorldType == EWorldType::Game
+		|| (World.WorldType == EWorldType::PIE && World.HasAnyFlags(RF_WasLoaded));
+#else
 	return true;
+#endif // WITH_EDITOR
 }
 
 void UMLAdapterManager::OnPostWorldInit(UWorld* World, const UWorld::InitializationValues)
