MLAdapter: Shorten log name to LogMLAdapter

#rnx

diff --git a/Engine/Plugins/AI/MLAdapter/Source/MLAdapter/Private/Agents/MLAdapterAgent.cpp b/Engine/Plugins/AI/MLAdapter/Source/MLAdapter/Private/Agents/MLAdapterAgent.cpp
--- a/Engine/Plugins/AI/MLAdapter/Source/MLAdapter/Private/Agents/MLAdapterAgent.cpp
+++ b/Engine/Plugins/AI/MLAdapter/Source/MLAdapter/Private/Agents/MLAdapterAgent.cpp
@@ -367,7 +367,7 @@ void UMLAdapterAgent::SetAvatar(AActor* InAvatar)
 
 	if (InAvatar != nullptr && IsSuitableAvatar(*InAvatar) == false)
 	{
-		UE_LOG(LogUnrealEditorMLAdapter, Log, TEXT("SetAvatar was called for agent %u but %s is not a valid avatar (required avatar class %s)"),
+		UE_LOG(LogMLAdapter, Log, TEXT("SetAvatar was called for agent %u but %s is not a valid avatar (required avatar class %s)"),
 			AgentID, *InAvatar->GetName(), *GetNameSafe(AvatarClass));
 		return;
 	}
diff --git a/Engine/Plugins/AI/MLAdapter/Source/MLAdapter/Private/Agents/MLAdapterAgent_Inference.cpp b/Engine/Plugins/AI/MLAdapter/Source/MLAdapter/Private/Agents/MLAdapterAgent_Inference.cpp
--- a/Engine/Plugins/AI/MLAdapter/Source/MLAdapter/Private/Agents/MLAdapterAgent_Inference.cpp
+++ b/Engine/Plugins/AI/MLAdapter/Source/MLAdapter/Private/Agents/MLAdapterAgent_Inference.cpp
@@ -20,7 +20,7 @@ void UMLAdapterAgent_Inference::Think(const float DeltaTime)
 {
 	if (Brain == nullptr)
 	{
-		UE_LOG(LogUnrealEditorMLAdapter, Warning, TEXT("Agent beginning to Think but Brain is null, Skipping."));
+		UE_LOG(LogMLAdapter, Warning, TEXT("Agent beginning to Think but Brain is null, Skipping."));
 		return;
 	}
 
diff --git a/Engine/Plugins/AI/MLAdapter/Source/MLAdapter/Private/MLAdapterLibrarian.cpp b/Engine/Plugins/AI/MLAdapter/Source/MLAdapter/Private/MLAdapterLibrarian.cpp
--- a/Engine/Plugins/AI/MLAdapter/Source/MLAdapter/Private/MLAdapterLibrarian.cpp
+++ b/Engine/Plugins/AI/MLAdapter/Source/MLAdapter/Private/MLAdapterLibrarian.cpp
@@ -20,7 +20,7 @@ namespace FLibrarianHelper
 
 		if (Class->HasAnyClassFlags(CLASS_Abstract | CLASS_Deprecated | CLASS_NewerVersionExists))
 		{
-			UE_LOG(LogUnrealEditorMLAdapter, Log, TEXT("Librarian: skipping class %s registration due to it being %s%s%s")
+			UE_LOG(LogMLAdapter, Log, TEXT("Librarian: skipping class %s registration due to it being %s%s%s")
 				, *Class->GetName()
 				, Class->HasAnyClassFlags(CLASS_Abstract) ? TEXT("Abstract, ") : TEXT("")
 				, Class->HasAnyClassFlags(CLASS_Deprecated) ? TEXT("Deprecated, ") : TEXT("")
diff --git a/Engine/Plugins/AI/MLAdapter/Source/MLAdapter/Private/MLAdapterTypes.cpp b/Engine/Plugins/AI/MLAdapter/Source/MLAdapter/Private/MLAdapterTypes.cpp
--- a/Engine/Plugins/AI/MLAdapter/Source/MLAdapter/Private/MLAdapterTypes.cpp
+++ b/Engine/Plugins/AI/MLAdapter/Source/MLAdapter/Private/MLAdapterTypes.cpp
@@ -4,7 +4,7 @@
 #include "GameFramework/Controller.h"
 #include "GameFramework/Pawn.h"
 
-DEFINE_LOG_CATEGORY(LogUnrealEditorMLAdapter);
+DEFINE_LOG_CATEGORY(LogMLAdapter);
 
 namespace FMLAdapter
 {
diff --git a/Engine/Plugins/AI/MLAdapter/Source/MLAdapter/Private/Managers/MLAdapterManager.cpp b/Engine/Plugins/AI/MLAdapter/Source/MLAdapter/Private/Managers/MLAdapterManager.cpp
--- a/Engine/Plugins/AI/MLAdapter/Source/MLAdapter/Private/Managers/MLAdapterManager.cpp
+++ b/Engine/Plugins/AI/MLAdapter/Source/MLAdapter/Private/Managers/MLAdapterManager.cpp
@@ -65,7 +65,7 @@ void UMLAdapterManager::RecreateManagerInstance()
 		? SettingsManagerClass.Get()
 		: UMLAdapterManager::StaticClass();
 
-	UE_LOG(LogUnrealEditorMLAdapter, Log, TEXT("Creating MLAdapter manager of class %s"), *GetNameSafe(Class));
+	UE_LOG(LogMLAdapter, Log, TEXT("Creating MLAdapter manager of class %s"), *GetNameSafe(Class));
 
 	UMLAdapterManager* NewInstance = NewObject<UMLAdapterManager>(GEngine, Class);
 	check(NewInstance);
@@ -148,7 +148,7 @@ void UMLAdapterManager::StopServer()
 {
 	if (RPCServerInstance)
 	{
-		UE_LOG(LogUnrealEditorMLAdapter, Log, TEXT("Stopping RPC server."));
+		UE_LOG(LogMLAdapter, Log, TEXT("Stopping RPC server."));
 		RPCServerInstance->stop();
 		delete RPCServerInstance;
 		RPCServerInstance = nullptr;
@@ -184,7 +184,7 @@ void UMLAdapterManager::StartServer(uint16 Port, EMLAdapterServerMode InMode, ui
 		}
 	}
 
-	UE_LOG(LogUnrealEditorMLAdapter, Log, TEXT("Starting RPC server on port %d."), Port);
+	UE_LOG(LogMLAdapter, Log, TEXT("Starting RPC server on port %d."), Port);
 	RPCServerInstance = new FRPCServer(Port);
 	CurrentPort = Port;
 	check(RPCServerInstance);
@@ -278,7 +278,7 @@ UMLAdapterSession* UMLAdapterManager::CreateNewSession()
 		? UMLAdapterSettings::GetSessionClass().Get()
 		: UMLAdapterSession::StaticClass();
 
-	UE_LOG(LogUnrealEditorMLAdapter, Log, TEXT("Creating MLAdapter session of class %s"), *GetNameSafe(Class));
+	UE_LOG(LogMLAdapter, Log, TEXT("Creating MLAdapter session of class %s"), *GetNameSafe(Class));
 
 	UMLAdapterSession* NewSession = FMLAdapter::NewObject<UMLAdapterSession>(this, Class);
 	check(NewSession);
diff --git a/Engine/Plugins/AI/MLAdapter/Source/MLAdapter/Private/Managers/MLAdapterManager_Client.cpp b/Engine/Plugins/AI/MLAdapter/Source/MLAdapter/Private/Managers/MLAdapterManager_Client.cpp
--- a/Engine/Plugins/AI/MLAdapter/Source/MLAdapter/Private/Managers/MLAdapterManager_Client.cpp
+++ b/Engine/Plugins/AI/MLAdapter/Source/MLAdapter/Private/Managers/MLAdapterManager_Client.cpp
@@ -94,7 +94,7 @@ void UMLAdapterManager::AddCommonFunctions(FRPCServer& Server)
 
 void UMLAdapterManager::ConfigureAsClient(FRPCServer& Server)
 {
-	UE_LOG(LogUnrealEditorMLAdapter, Log, TEXT("\tconfiguring as client"));
+	UE_LOG(LogMLAdapter, Log, TEXT("\tconfiguring as client"));
 
 	AddCommonFunctions(Server);
 
diff --git a/Engine/Plugins/AI/MLAdapter/Source/MLAdapter/Private/Managers/MLAdapterManager_Server.cpp b/Engine/Plugins/AI/MLAdapter/Source/MLAdapter/Private/Managers/MLAdapterManager_Server.cpp
--- a/Engine/Plugins/AI/MLAdapter/Source/MLAdapter/Private/Managers/MLAdapterManager_Server.cpp
+++ b/Engine/Plugins/AI/MLAdapter/Source/MLAdapter/Private/Managers/MLAdapterManager_Server.cpp
@@ -22,7 +22,7 @@
 
 void UMLAdapterManager::ConfigureAsServer(FRPCServer& Server)
 {
-	UE_LOG(LogUnrealEditorMLAdapter, Log, TEXT("\tconfiguring as server"));
+	UE_LOG(LogMLAdapter, Log, TEXT("\tconfiguring as server"));
 
 	AddCommonFunctions(Server);
 
diff --git a/Engine/Plugins/AI/MLAdapter/Source/MLAdapter/Private/Sensors/MLAdapterSensor_EnhancedInput.cpp b/Engine/Plugins/AI/MLAdapter/Source/MLAdapter/Private/Sensors/MLAdapterSensor_EnhancedInput.cpp
--- a/Engine/Plugins/AI/MLAdapter/Source/MLAdapter/Private/Sensors/MLAdapterSensor_EnhancedInput.cpp
+++ b/Engine/Plugins/AI/MLAdapter/Source/MLAdapter/Private/Sensors/MLAdapterSensor_EnhancedInput.cpp
@@ -80,7 +80,7 @@ void UMLAdapterSensor_EnhancedInput::OnAvatarSet(AActor* Avatar)
 
 	if (InputComponent != nullptr)
 	{
-		UE_LOG(LogUnrealEditorMLAdapter, Log, TEXT("MLAdapterSensor_EnhancedInput: This sensor already bound to another InputComponent. Cleaning up previous InputComponent and rebinding to new avatar."));
+		UE_LOG(LogMLAdapter, Log, TEXT("MLAdapterSensor_EnhancedInput: This sensor already bound to another InputComponent. Cleaning up previous InputComponent and rebinding to new avatar."));
 		InputComponent->DestroyComponent();
 	}
 
@@ -108,13 +108,13 @@ void UMLAdapterSensor_EnhancedInput::OnAvatarSet(AActor* Avatar)
 			}
 
 			UInputDelegateBinding::BindInputDelegates(Pawn->GetClass(), InputComponent);
-			UE_LOG(LogUnrealEditorMLAdapter, Log, TEXT("MLAdapterSensor_EnhancedInput: Successfully bound to %s"), *GetNameSafe(PawnClass));
+			UE_LOG(LogMLAdapter, Log, TEXT("MLAdapterSensor_EnhancedInput: Successfully bound to %s"), *GetNameSafe(PawnClass));
 		}
 	}
 	else
 	{
 		FString AvatarName = Avatar->GetName();
-		UE_LOG(LogUnrealEditorMLAdapter, Warning, TEXT("MLAdapterSensor_EnhancedInput: Unable to bind - could not cast to APawn or APlayerController with a Pawn. Avatar name was %s"), *AvatarName);
+		UE_LOG(LogMLAdapter, Warning, TEXT("MLAdapterSensor_EnhancedInput: Unable to bind - could not cast to APawn or APlayerController with a Pawn. Avatar name was %s"), *AvatarName);
 	}
 
 	Super::OnAvatarSet(Avatar);
diff --git a/Engine/Plugins/AI/MLAdapter/Source/MLAdapter/Private/Sessions/MLAdapterLocalDataCollectionSession.cpp b/Engine/Plugins/AI/MLAdapter/Source/MLAdapter/Private/Sessions/MLAdapterLocalDataCollectionSession.cpp
--- a/Engine/Plugins/AI/MLAdapter/Source/MLAdapter/Private/Sessions/MLAdapterLocalDataCollectionSession.cpp
+++ b/Engine/Plugins/AI/MLAdapter/Source/MLAdapter/Private/Sessions/MLAdapterLocalDataCollectionSession.cpp
@@ -36,7 +36,7 @@ void UMLAdapterLocalDataCollectionSession::Tick(float DeltaTime)
 
 	if (Agent ==  nullptr)
 	{
-		UE_LOG(LogUnrealEditorMLAdapter, Log, TEXT("LocalDataCollectionSession: Player-controlled agent not found yet."))
+		UE_LOG(LogMLAdapter, Log, TEXT("LocalDataCollectionSession: Player-controlled agent not found yet."))
 		return;
 	}
 
diff --git a/Engine/Plugins/AI/MLAdapter/Source/MLAdapter/Private/Sessions/MLAdapterSession.cpp b/Engine/Plugins/AI/MLAdapter/Source/MLAdapter/Private/Sessions/MLAdapterSession.cpp
--- a/Engine/Plugins/AI/MLAdapter/Source/MLAdapter/Private/Sessions/MLAdapterSession.cpp
+++ b/Engine/Plugins/AI/MLAdapter/Source/MLAdapter/Private/Sessions/MLAdapterSession.cpp
@@ -347,7 +347,7 @@ FMLAdapter::FAgentID UMLAdapterSession::AddAgent()
 		? UMLAdapterSettings::GetAgentClass().Get()
 		: UMLAdapterAgent::StaticClass();
 
-	UE_LOG(LogUnrealEditorMLAdapter, Log, TEXT("Creating MLAdapter agent of class %s"), *GetNameSafe(AgentClass));
+	UE_LOG(LogMLAdapter, Log, TEXT("Creating MLAdapter agent of class %s"), *GetNameSafe(AgentClass));
 
 	UMLAdapterAgent* NewAgent = FMLAdapter::NewObject<UMLAdapterAgent>(this, AgentClass);
 
@@ -396,7 +396,7 @@ UMLAdapterAgent* UMLAdapterSession::GetAgent(FMLAdapter::FAgentID AgentID)
 {
 	if (Agents.IsValidIndex(AgentID) == false)
 	{
-		UE_LOG(LogUnrealEditorMLAdapter, Warning, TEXT("UMLAdapterSession::GetAgent: Invalid AgentID. Failing"));
+		UE_LOG(LogMLAdapter, Warning, TEXT("UMLAdapterSession::GetAgent: Invalid AgentID. Failing"));
 		return nullptr;
 	}
 
@@ -489,7 +489,7 @@ bool UMLAdapterSession::RequestAvatarForAgent(UMLAdapterAgent& Agent, UWorld* In
 	if (Agent.GetAvatar() != nullptr)
 	{
 		// skipping.
-		UE_LOG(LogUnrealEditorMLAdapter, Verbose, TEXT("UMLAdapterSession::RequestAvatarForAgent called for agent [%s] while it still has an avatar [%s]. Call ClearAvatar first to null-out agent\'s avatar."),
+		UE_LOG(LogMLAdapter, Verbose, TEXT("UMLAdapterSession::RequestAvatarForAgent called for agent [%s] while it still has an avatar [%s]. Call ClearAvatar first to null-out agent\'s avatar."),
 			Agent.GetAgentID(), *GetNameSafe(Agent.GetAvatar()));
 		return false;
 	}
@@ -507,7 +507,7 @@ bool UMLAdapterSession::RequestAvatarForAgent(UMLAdapterAgent& Agent, UWorld* In
 	
 	if (InWorld == nullptr)
 	{
-		UE_LOG(LogUnrealEditorMLAdapter, Warning, TEXT("UMLAdapterSession::RequestAvatarForAgent called with InWorld and CachedWorld both being null. Auto-failure."));
+		UE_LOG(LogMLAdapter, Warning, TEXT("UMLAdapterSession::RequestAvatarForAgent called with InWorld and CachedWorld both being null. Auto-failure."));
 		return false;
 	}
 
diff --git a/Engine/Plugins/AI/MLAdapter/Source/MLAdapter/Public/MLAdapterTypes.h b/Engine/Plugins/AI/MLAdapter/Source/MLAdapter/Public/MLAdapterTypes.h
--- a/Engine/Plugins/AI/MLAdapter/Source/MLAdapter/Public/MLAdapterTypes.h
+++ b/Engine/Plugins/AI/MLAdapter/Source/MLAdapter/Public/MLAdapterTypes.h
@@ -11,7 +11,7 @@
 class AActor;
 class AController;
 
-DECLARE_LOG_CATEGORY_EXTERN(LogUnrealEditorMLAdapter, Log, All);
+DECLARE_LOG_CATEGORY_EXTERN(LogMLAdapter, Log, All);
 
 #ifndef DEFINE_ENUM_TO_STRING
 #define DEFINE_ENUM_TO_STRING(EnumType) FString EnumToString(const EnumType Value) \
