MLAdapter: Add filepath & timestamp options to local data collection session

#rnx

diff --git a/Engine/Plugins/AI/MLAdapter/Source/MLAdapter/Private/Sessions/MLAdapterLocalDataCollectionSession.cpp b/Engine/Plugins/AI/MLAdapter/Source/MLAdapter/Private/Sessions/MLAdapterLocalDataCollectionSession.cpp
--- a/Engine/Plugins/AI/MLAdapter/Source/MLAdapter/Private/Sessions/MLAdapterLocalDataCollectionSession.cpp
+++ b/Engine/Plugins/AI/MLAdapter/Source/MLAdapter/Private/Sessions/MLAdapterLocalDataCollectionSession.cpp
@@ -6,6 +6,37 @@
 #include "Misc/FileHelper.h"
 #include "HAL/FileManager.h"
 
+void UMLAdapterLocalDataCollectionSession::PostInitProperties()
+{
+	Super::PostInitProperties();
+
+	if (HasAnyFlags(RF_ClassDefaultObject))
+	{
+		return;
+	}
+
+	if (!FilePath.Path.IsEmpty())
+	{
+		OutputFilePath += FilePath.Path + FGenericPlatformMisc::GetDefaultPathSeparator();
+	}
+
+	if (bPrefixOutputFilenameWithTimestamp)
+	{
+		OutputFilePath += FDateTime::Now().ToString() + "_";
+	}
+
+	if (!FileName.IsEmpty())
+	{
+		OutputFilePath += FileName;
+	}
+	else
+	{
+		OutputFilePath += "mladapter.data";
+	}
+
+	UE_LOG(LogMLAdapter, Log, TEXT("LocalDataCollectionSession: Writing data to %s"), *OutputFilePath);
+}
+
 void UMLAdapterLocalDataCollectionSession::OnPostWorldInit(UWorld& World)
 {
 	Super::OnPostWorldInit(World);
@@ -36,7 +67,7 @@ void UMLAdapterLocalDataCollectionSession::Tick(float DeltaTime)
 
 	if (Agent ==  nullptr)
 	{
-		UE_LOG(LogMLAdapter, Log, TEXT("LocalDataCollectionSession: Player-controlled agent not found yet."))
+		UE_LOG(LogMLAdapter, Log, TEXT("LocalDataCollectionSession: Player-controlled agent not found yet."));
 		return;
 	}
 
@@ -45,7 +76,11 @@ void UMLAdapterLocalDataCollectionSession::Tick(float DeltaTime)
 	FMLAdapterMemoryWriter Writer(Buffer);
 	Agent->GetObservations(Writer);
 
-	FFileHelper::SaveArrayToFile(Buffer, *FileName, &IFileManager::Get(), EFileWrite::FILEWRITE_Append);
+	const bool bSuccess = FFileHelper::SaveArrayToFile(Buffer, *OutputFilePath, &IFileManager::Get(), EFileWrite::FILEWRITE_Append);
+	if (!bSuccess)
+	{
+		UE_LOG(LogMLAdapter, Error, TEXT("LocalDataCollectionSession: Failed to write data to output file."));
+	}
 }
 
 void UMLAdapterLocalDataCollectionSession::Close()
diff --git a/Engine/Plugins/AI/MLAdapter/Source/MLAdapter/Public/Sessions/MLAdapterLocalDataCollectionSession.h b/Engine/Plugins/AI/MLAdapter/Source/MLAdapter/Public/Sessions/MLAdapterLocalDataCollectionSession.h
--- a/Engine/Plugins/AI/MLAdapter/Source/MLAdapter/Public/Sessions/MLAdapterLocalDataCollectionSession.h
+++ b/Engine/Plugins/AI/MLAdapter/Source/MLAdapter/Public/Sessions/MLAdapterLocalDataCollectionSession.h
@@ -3,6 +3,9 @@
 #pragma once 
 
 #include "Sessions/MLAdapterSession.h"
+#include "Engine/EngineTypes.h"
+#include "GenericPlatform/GenericPlatformMisc.h"
+#include "Misc/DateTime.h"
 #include "UObject/WeakObjectPtrTemplates.h"
 #include "MLAdapterLocalDataCollectionSession.generated.h"
 
@@ -15,6 +18,8 @@ class MLADAPTER_API UMLAdapterLocalDataCollectionSession : public UMLAdapterSess
 {
 	GENERATED_BODY()
 public:
+	virtual void PostInitProperties() override;
+
 	virtual void OnPostWorldInit(UWorld& World) override;
 
 	UFUNCTION()
@@ -24,10 +29,17 @@ public:
 
 	virtual void Close() override;
 
-	UPROPERTY(EditAnywhere, Category = MLAdapter)
+	UPROPERTY(EditDefaultsOnly, Category = MLAdapter)
+	FDirectoryPath FilePath;
+
+	UPROPERTY(EditDefaultsOnly, Category = MLAdapter)
 	FString FileName;
 
+	UPROPERTY(EditDefaultsOnly, Category = MLAdapter)
+	bool bPrefixOutputFilenameWithTimestamp;
+
 private:
+	FString OutputFilePath;
 
 	UPROPERTY()
 	TWeakObjectPtr<UMLAdapterAgent> PlayerControlledAgent;
