Preparation pass for TObjectPtr upgrade pass on plugins used when building UnrealEditor with "-allmodules".  Addresses missed header issue and Houdini plugin issues.

diff --git a/Engine/Plugins/AI/MLAdapter/Source/MLAdapter/Public/Agents/MLAdapterAgent.h b/Engine/Plugins/AI/MLAdapter/Source/MLAdapter/Public/Agents/MLAdapterAgent.h
--- a/Engine/Plugins/AI/MLAdapter/Source/MLAdapter/Public/Agents/MLAdapterAgent.h
+++ b/Engine/Plugins/AI/MLAdapter/Source/MLAdapter/Public/Agents/MLAdapterAgent.h
@@ -126,7 +126,7 @@ public:
 	/** Get the actuator with the given ID if this agent has it. */
 	UMLAdapterActuator* GetActuator(const uint32 ActuatorID) 
 	{ 
-		UMLAdapterActuator** FoundActuator = Actuators.FindByPredicate([ActuatorID](const UMLAdapterActuator* Actuator) { return (Actuator->GetElementID() == ActuatorID); });
+		UE_TRANSITIONAL_OBJECT_PTR(UMLAdapterActuator)* FoundActuator = Actuators.FindByPredicate([ActuatorID](const UMLAdapterActuator* Actuator) { return (Actuator->GetElementID() == ActuatorID); });
 		return FoundActuator ? *FoundActuator : nullptr;
 	}
 
diff --git a/Engine/Plugins/Runtime/nDisplay/Source/DisplayClusterTests/Private/DisplayCluster.spec.cpp b/Engine/Plugins/Runtime/nDisplay/Source/DisplayClusterTests/Private/DisplayCluster.spec.cpp
--- a/Engine/Plugins/Runtime/nDisplay/Source/DisplayClusterTests/Private/DisplayCluster.spec.cpp
+++ b/Engine/Plugins/Runtime/nDisplay/Source/DisplayClusterTests/Private/DisplayCluster.spec.cpp
@@ -359,7 +359,7 @@ void FDisplayClusterSpec::Define()
 			{
 				if (TestTrue("Instance has entry corresponding to CDO key", InstancedRootCluster->Nodes.Contains(Pair.Key)))
 				{
-					TestNotNull("Instance has valid entry corresponding to CDO key", InstancedRootCluster->Nodes[Pair.Key]);
+					TestNotNull("Instance has valid entry corresponding to CDO key", ToRawPtr(InstancedRootCluster->Nodes[Pair.Key]));
 				}
 			}
 		});
