[MLAdapter]
* MLAdapterManager - Fixed non-unity/pch compile error

diff --git a/Engine/Plugins/AI/MLAdapter/Source/MLAdapter/Private/Managers/MLAdapterManager.cpp b/Engine/Plugins/AI/MLAdapter/Source/MLAdapter/Private/Managers/MLAdapterManager.cpp
--- a/Engine/Plugins/AI/MLAdapter/Source/MLAdapter/Private/Managers/MLAdapterManager.cpp
+++ b/Engine/Plugins/AI/MLAdapter/Source/MLAdapter/Private/Managers/MLAdapterManager.cpp
@@ -2,6 +2,7 @@
 
 #include "Managers/MLAdapterManager.h"
 #include "CoreGlobals.h"
+#include "Misc/CommandLine.h"
 #include "Misc/CoreDelegates.h"
 #include "Engine/Engine.h"
 #include "GameFramework/GameModeBase.h"
