ML Adapter: Change from beta to experimental. See https://epicgames.slack.com/archives/C045565NV1S/p1673365091617639
#rnx

diff --git a/Engine/Plugins/AI/MLAdapter/MLAdapter.uplugin b/Engine/Plugins/AI/MLAdapter/MLAdapter.uplugin
--- a/Engine/Plugins/AI/MLAdapter/MLAdapter.uplugin
+++ b/Engine/Plugins/AI/MLAdapter/MLAdapter.uplugin
@@ -3,7 +3,7 @@
 	"Version": 1,
 	"VersionName": "0.0.1",
 	"FriendlyName": "ML Adapter",
-	"Description": "[EXPERIMENTAL] A framework for training and utilizing machine learning agents in games. Creates an RPC interface through which an external process can query game state and control in-game actors. Once trained, agents can be run in-engine via neural networks loaded from ONNX models.",
+	"Description": "A framework for training and utilizing machine learning agents in games. Creates an RPC interface through which an external process can query game state and control in-game actors. Once trained, agents can be run in-engine via neural networks loaded from ONNX models.",
 	"Category": "AI",
 	"CreatedBy": "Epic Games, Inc.",
 	"CreatedByURL": "https://epicgames.com",
@@ -12,7 +12,8 @@
 	"SupportURL": "",
 	"EnabledByDefault": false,
 	"CanContainContent": false,
-	"IsBetaVersion": true,
+	"IsBetaVersion": false,
+	"IsExperimentalVersion": true,
 	"Installed": false,
 	"Modules": [
 		{
