MLAdapter - Upgrade to Neural Network Engine (NNE)

diff --git a/Engine/Plugins/AI/MLAdapter/MLAdapter.uplugin b/Engine/Plugins/AI/MLAdapter/MLAdapter.uplugin
--- a/Engine/Plugins/AI/MLAdapter/MLAdapter.uplugin
+++ b/Engine/Plugins/AI/MLAdapter/MLAdapter.uplugin
@@ -37,7 +37,7 @@
 			"Enabled": true
 		},
 		{
-			"Name": "NeuralNetworkInference",
+			"Name": "NNE",
 			"Enabled": true
 		}
 	]
diff --git a/Engine/Plugins/AI/MLAdapter/Source/MLAdapter/MLAdapter.Build.cs b/Engine/Plugins/AI/MLAdapter/Source/MLAdapter/MLAdapter.Build.cs
--- a/Engine/Plugins/AI/MLAdapter/Source/MLAdapter/MLAdapter.Build.cs
+++ b/Engine/Plugins/AI/MLAdapter/Source/MLAdapter/MLAdapter.Build.cs
@@ -31,7 +31,7 @@ namespace UnrealBuildTool.Rules
                     "Json",
                     "JsonUtilities",
                     "GameplayAbilities",
-					"NeuralNetworkInference"
+					"NNECore"
 				}
             );
 
diff --git a/Engine/Plugins/AI/MLAdapter/Source/MLAdapter/Private/Agents/MLAdapterAgent_Inference.cpp b/Engine/Plugins/AI/MLAdapter/Source/MLAdapter/Private/Agents/MLAdapterAgent_Inference.cpp
--- a/Engine/Plugins/AI/MLAdapter/Source/MLAdapter/Private/Agents/MLAdapterAgent_Inference.cpp
+++ b/Engine/Plugins/AI/MLAdapter/Source/MLAdapter/Private/Agents/MLAdapterAgent_Inference.cpp
@@ -1,19 +1,75 @@
 // Copyright Epic Games, Inc. All Rights Reserved.
 
 #include "Agents/MLAdapterAgent_Inference.h"
-#include "MLAdapterTypes.h"
 #include "HAL/UnrealMemory.h"
+#include "MLAdapterTypes.h"
+#include "NNECore.h"
+#include "NNECoreTypes.h"
 
 void UMLAdapterAgent_Inference::PostInitProperties()
 {
-	UNeuralNetwork* NeuralNetwork = (UNeuralNetwork*)NeuralNetworkPath.TryLoad();
+	const FString RuntimeName = TEXT("NNERuntimeORTCpu");
+
+	Super::PostInitProperties();
 	
-	if (NeuralNetwork != nullptr)
+	if (ModelData == nullptr)
 	{
-		Brain = NeuralNetwork;
+		UE_LOG(LogMLAdapter, Warning, TEXT("Neural network asset data not set"));
+		return;
 	}
 
-	Super::PostInitProperties();
+	TWeakInterfacePtr<INNERuntimeCPU> Runtime = UE::NNECore::GetRuntime<INNERuntimeCPU>(RuntimeName);
+	if (!Runtime.IsValid())
+	{
+		UE_LOG(LogMLAdapter, Warning, TEXT("Runtime %s is not valid!"), *RuntimeName);
+		return;
+	}
+
+	Brain = Runtime->CreateModelCPU(ModelData);
+
+	TConstArrayView<UE::NNECore::FTensorDesc> InputTensorDescs = Brain->GetInputTensorDescs();
+	if (InputTensorDescs.Num() != 1)
+	{
+		UE_LOG(LogMLAdapter, Warning, TEXT("Brain accepts only single tensor input."));
+		return;
+	}
+	if (!InputTensorDescs[0].GetShape().IsConcrete())
+	{
+		UE_LOG(LogMLAdapter, Warning, TEXT("Brain does not accept dynamic tensor input shapes."));
+		return;
+	}
+
+	UE::NNECore::FTensorShape InputTensorShape = UE::NNECore::FTensorShape::MakeFromSymbolic(InputTensorDescs[0].GetShape());
+
+	Brain->SetInputTensorShapes({ InputTensorShape });
+
+	TConstArrayView<UE::NNECore::FTensorDesc> OutputTensorDescs = Brain->GetOutputTensorDescs();
+	if (OutputTensorDescs.Num() != 1)
+	{
+		UE_LOG(LogMLAdapter, Warning, TEXT("Brain accepts only single tensor output."));
+		return;
+	}
+
+	TConstArrayView<UE::NNECore::FTensorShape> OutputTensorShapes = Brain->GetOutputTensorShapes();
+	if (OutputTensorShapes.Num() != 1)
+	{
+		UE_LOG(LogMLAdapter, Warning, TEXT("Brain could not resolve output shapes."));
+		return;
+	}
+
+	InputTensorSizeInBytes = InputTensorDescs[0].GetElemByteSize() * InputTensorShape.Volume();
+	if (InputTensorSizeInBytes == 0)
+	{
+		UE_LOG(LogMLAdapter, Warning, TEXT("Input tensor should be of non-zero size."));
+		return;
+	}
+
+	OutputTensorSizeInBytes = OutputTensorDescs[0].GetElemByteSize() * OutputTensorShapes[0].Volume();
+	if (OutputTensorSizeInBytes == 0)
+	{
+		UE_LOG(LogMLAdapter, Warning, TEXT("Output tensor should be of non-zero size."));
+		return;
+	}
 }
 
 void UMLAdapterAgent_Inference::Think(const float DeltaTime)
@@ -30,19 +86,31 @@ void UMLAdapterAgent_Inference::Think(const float DeltaTime)
 		return;
 	}
 
+	if (InputTensorSizeInBytes == 0 || OutputTensorSizeInBytes == 0)
+	{
+		return;
+	}
+
 	bActionDurationElapsed = false;
 
 	TArray<uint8> Buffer;
 	FMLAdapterMemoryWriter Writer(Buffer);
 	GetObservations(Writer);
 
-	const float* DataPtr = (float*)Buffer.GetData();
-	FMemory::Memcpy(Brain->GetInputDataPointerMutable(), DataPtr, Buffer.Num());
+	if ((uint32)Buffer.Num() < InputTensorSizeInBytes)
+	{
+		UE_LOG(LogMLAdapter, Warning, TEXT("Brain received not enough data."));
+		return;
+	}
+	
+	TArray<uint8> OutputBuffer;
+	OutputBuffer.SetNumUninitialized(OutputTensorSizeInBytes);
+
+	UE::NNECore::FTensorBindingCPU InputTensorBinding{Buffer.GetData(), Buffer.Num()};
+	UE::NNECore::FTensorBindingCPU OutputTensorBinding{OutputBuffer.GetData(), OutputBuffer.Num()};
 
-	Brain->Run();
+	Brain->RunSync({InputTensorBinding}, {OutputTensorBinding});
 
-	const FNeuralTensor& Tensor = Brain->GetOutputTensor();
-	const TArray<uint8>& Data = Tensor.GetUnderlyingUInt8ArrayRef();
-	FMLAdapterMemoryReader Reader(Data);
+	FMLAdapterMemoryReader Reader(OutputBuffer);
 	DigestActions(Reader);
 }
diff --git a/Engine/Plugins/AI/MLAdapter/Source/MLAdapter/Public/Agents/MLAdapterAgent_Inference.h b/Engine/Plugins/AI/MLAdapter/Source/MLAdapter/Public/Agents/MLAdapterAgent_Inference.h
--- a/Engine/Plugins/AI/MLAdapter/Source/MLAdapter/Public/Agents/MLAdapterAgent_Inference.h
+++ b/Engine/Plugins/AI/MLAdapter/Source/MLAdapter/Public/Agents/MLAdapterAgent_Inference.h
@@ -2,9 +2,11 @@
 
 #pragma once
 
+#include "Templates/UniquePtr.h"
 #include "UObject/SoftObjectPath.h"
 #include "Agents/MLAdapterAgent.h"
-#include "NeuralNetwork.h"
+#include "NNECoreModelData.h"
+#include "NNECoreRuntimeCPU.h"
 #include "MLAdapterAgent_Inference.generated.h"
 
 /**
@@ -17,17 +19,20 @@ class MLADAPTER_API UMLAdapterAgent_Inference : public UMLAdapterAgent
 {
 	GENERATED_BODY()
 public:
-	/** Tries to load the neural network from the NeuralNetworkPath. */
+	/** Tries to load the neural network from the ModelDataPath. */
 	virtual void PostInitProperties() override;
 
 	/** Process observations from the sensors with the neural network and populate the actuators with data. */
 	virtual void Think(const float DeltaTime) override;
 
-	/** The path to the neural network asset to be loaded. */
-	UPROPERTY(EditDefaultsOnly, Category = MLAdapter, meta = (AllowedClasses = "NeuralNetwork"))
-	FSoftObjectPath NeuralNetworkPath;
+	/** Reference to neural network asset data. */
+	UPROPERTY(EditAnywhere, Category = MLAdapter)
+	TObjectPtr<UNNEModelData> ModelData;
 
 	/** The neural network controlling this agent. */
-	UPROPERTY()
-	TObjectPtr<UNeuralNetwork> Brain;
+	TUniquePtr<UE::NNECore::IModelCPU> Brain;
+
+private:
+	uint32 InputTensorSizeInBytes = 0;
+	uint32 OutputTensorSizeInBytes = 0;
 };
