Removing redundant include paths.

diff --git a/Engine/Plugins/AI/MLAdapter/Source/MLAdapter/MLAdapter.Build.cs b/Engine/Plugins/AI/MLAdapter/Source/MLAdapter/MLAdapter.Build.cs
--- a/Engine/Plugins/AI/MLAdapter/Source/MLAdapter/MLAdapter.Build.cs
+++ b/Engine/Plugins/AI/MLAdapter/Source/MLAdapter/MLAdapter.Build.cs
@@ -1,10 +1,8 @@
 // Copyright Epic Games, Inc. All Rights Reserved.
 
-using System.IO;
-
 namespace UnrealBuildTool.Rules
 {
-    public class MLAdapter : ModuleRules
+	public class MLAdapter : ModuleRules
     {
         public MLAdapter(ReadOnlyTargetRules Target) : base(Target)
         {
@@ -12,13 +10,6 @@ namespace UnrealBuildTool.Rules
             bEnableExceptions = true;
             PCHUsage = ModuleRules.PCHUsageMode.UseExplicitOrSharedPCHs;
 
-            PublicIncludePaths.AddRange(
-                new string[] {
-                "Runtime/AIModule/Public",
-                ModuleDirectory + "/Public",
-                }
-            );
-
             PublicDependencyModuleNames.AddRange(
                 new string[] {
                     "Core",
diff --git a/Engine/Plugins/AI/MassAI/Source/MassAIDebug/MassAIDebug.Build.cs b/Engine/Plugins/AI/MassAI/Source/MassAIDebug/MassAIDebug.Build.cs
--- a/Engine/Plugins/AI/MassAI/Source/MassAIDebug/MassAIDebug.Build.cs
+++ b/Engine/Plugins/AI/MassAI/Source/MassAIDebug/MassAIDebug.Build.cs
@@ -8,13 +8,6 @@ namespace UnrealBuildTool.Rules
 		{
 			PCHUsage = ModuleRules.PCHUsageMode.UseExplicitOrSharedPCHs;
 
-			PublicIncludePaths.AddRange(
-				new string[] {
-					"Runtime/AIModule/Public",
-					ModuleDirectory + "/Public",
-				}
-			);
-
 			PublicDependencyModuleNames.AddRange(
 				new string[] {
 					"Core",
diff --git a/Engine/Plugins/AI/MassAI/Source/MassAIReplication/MassAIReplication.Build.cs b/Engine/Plugins/AI/MassAI/Source/MassAIReplication/MassAIReplication.Build.cs
--- a/Engine/Plugins/AI/MassAI/Source/MassAIReplication/MassAIReplication.Build.cs
+++ b/Engine/Plugins/AI/MassAI/Source/MassAIReplication/MassAIReplication.Build.cs
@@ -8,13 +8,6 @@ namespace UnrealBuildTool.Rules
 		{
 			PCHUsage = ModuleRules.PCHUsageMode.UseExplicitOrSharedPCHs;
 
-			PublicIncludePaths.AddRange(
-				new string[] {
-				"Runtime/AIModule/Public",
-				ModuleDirectory + "/Public",
-				}
-			);
-
 			PublicDependencyModuleNames.AddRange(
 				new string[] {
 					"Core",
diff --git a/Engine/Plugins/AI/MassAI/Source/MassZoneGraphNavigation/MassZoneGraphNavigation.Build.cs b/Engine/Plugins/AI/MassAI/Source/MassZoneGraphNavigation/MassZoneGraphNavigation.Build.cs
--- a/Engine/Plugins/AI/MassAI/Source/MassZoneGraphNavigation/MassZoneGraphNavigation.Build.cs
+++ b/Engine/Plugins/AI/MassAI/Source/MassZoneGraphNavigation/MassZoneGraphNavigation.Build.cs
@@ -8,13 +8,6 @@ namespace UnrealBuildTool.Rules
 		{
 			PCHUsage = ModuleRules.PCHUsageMode.UseExplicitOrSharedPCHs;
 
-			PublicIncludePaths.AddRange(
-				new string[] {
-					"Runtime/AIModule/Public",
-					ModuleDirectory + "/Public",
-				}
-			);
-
 			PublicDependencyModuleNames.AddRange(
 				new string[] {
 					"AIModule",
diff --git a/Engine/Plugins/AI/MassCrowd/Source/MassCrowd/MassCrowd.Build.cs b/Engine/Plugins/AI/MassCrowd/Source/MassCrowd/MassCrowd.Build.cs
--- a/Engine/Plugins/AI/MassCrowd/Source/MassCrowd/MassCrowd.Build.cs
+++ b/Engine/Plugins/AI/MassCrowd/Source/MassCrowd/MassCrowd.Build.cs
@@ -8,13 +8,6 @@ namespace UnrealBuildTool.Rules
 		{
 			PCHUsage = ModuleRules.PCHUsageMode.UseExplicitOrSharedPCHs;
 
-			PublicIncludePaths.AddRange(
-				new string[] {
-					"Runtime/AIModule/Public",
-					ModuleDirectory + "/Public",
-				}
-			);
-
 			PublicDependencyModuleNames.AddRange(
 				new string[] {
 					"AIModule",
diff --git a/Engine/Plugins/Runtime/MassGameplay/Source/MassActors/MassActors.Build.cs b/Engine/Plugins/Runtime/MassGameplay/Source/MassActors/MassActors.Build.cs
--- a/Engine/Plugins/Runtime/MassGameplay/Source/MassActors/MassActors.Build.cs
+++ b/Engine/Plugins/Runtime/MassGameplay/Source/MassActors/MassActors.Build.cs
@@ -8,13 +8,6 @@ namespace UnrealBuildTool.Rules
 		{
 			PCHUsage = ModuleRules.PCHUsageMode.UseExplicitOrSharedPCHs;
 
-			PublicIncludePaths.AddRange(
-				new string[] {
-				"Runtime/AIModule/Public",
-				ModuleDirectory + "/Public",
-				}
-			);
-
 			PublicDependencyModuleNames.AddRange(
 				new string[] {
 					"Core",
diff --git a/Engine/Plugins/Runtime/MassGameplay/Source/MassCommon/MassCommon.Build.cs b/Engine/Plugins/Runtime/MassGameplay/Source/MassCommon/MassCommon.Build.cs
--- a/Engine/Plugins/Runtime/MassGameplay/Source/MassCommon/MassCommon.Build.cs
+++ b/Engine/Plugins/Runtime/MassGameplay/Source/MassCommon/MassCommon.Build.cs
@@ -8,13 +8,6 @@ namespace UnrealBuildTool.Rules
 		{
 			PCHUsage = ModuleRules.PCHUsageMode.UseExplicitOrSharedPCHs;
 
-			PublicIncludePaths.AddRange(
-				new string[] {
-				"Runtime/AIModule/Public",
-				ModuleDirectory + "/Public",
-				}
-			);
-
 			PublicDependencyModuleNames.AddRange(
 				new string[] {
 					"Core",
diff --git a/Engine/Plugins/Runtime/MassGameplay/Source/MassGameplayDebug/MassGameplayDebug.Build.cs b/Engine/Plugins/Runtime/MassGameplay/Source/MassGameplayDebug/MassGameplayDebug.Build.cs
--- a/Engine/Plugins/Runtime/MassGameplay/Source/MassGameplayDebug/MassGameplayDebug.Build.cs
+++ b/Engine/Plugins/Runtime/MassGameplay/Source/MassGameplayDebug/MassGameplayDebug.Build.cs
@@ -8,13 +8,6 @@ namespace UnrealBuildTool.Rules
 		{
 			PCHUsage = ModuleRules.PCHUsageMode.UseExplicitOrSharedPCHs;
 
-			PublicIncludePaths.AddRange(
-				new string[] {
-					"Runtime/AIModule/Public",
-					ModuleDirectory + "/Public",
-				}
-			);
-
 			PublicDependencyModuleNames.AddRange(
 				new string[] {
 					"Core",
diff --git a/Engine/Plugins/Runtime/MassGameplay/Source/MassLOD/MassLOD.Build.cs b/Engine/Plugins/Runtime/MassGameplay/Source/MassLOD/MassLOD.Build.cs
--- a/Engine/Plugins/Runtime/MassGameplay/Source/MassLOD/MassLOD.Build.cs
+++ b/Engine/Plugins/Runtime/MassGameplay/Source/MassLOD/MassLOD.Build.cs
@@ -8,13 +8,6 @@ namespace UnrealBuildTool.Rules
 		{
 			PCHUsage = ModuleRules.PCHUsageMode.UseExplicitOrSharedPCHs;
 
-			PublicIncludePaths.AddRange(
-				new string[] {
-				"Runtime/AIModule/Public",
-				ModuleDirectory + "/Public",
-				}
-			);
-
 			PublicDependencyModuleNames.AddRange(
 				new string[] {
 					"Core",
diff --git a/Engine/Plugins/Runtime/MassGameplay/Source/MassMovement/MassMovement.Build.cs b/Engine/Plugins/Runtime/MassGameplay/Source/MassMovement/MassMovement.Build.cs
--- a/Engine/Plugins/Runtime/MassGameplay/Source/MassMovement/MassMovement.Build.cs
+++ b/Engine/Plugins/Runtime/MassGameplay/Source/MassMovement/MassMovement.Build.cs
@@ -8,13 +8,6 @@ namespace UnrealBuildTool.Rules
 		{
 			PCHUsage = ModuleRules.PCHUsageMode.UseExplicitOrSharedPCHs;
 
-			PublicIncludePaths.AddRange(
-				new string[] {
-				"Runtime/AIModule/Public",
-				ModuleDirectory + "/Public",
-				}
-			);
-
 			PublicDependencyModuleNames.AddRange(
 				new string[] {
 					"AIModule",
diff --git a/Engine/Plugins/Runtime/MassGameplay/Source/MassReplication/MassReplication.Build.cs b/Engine/Plugins/Runtime/MassGameplay/Source/MassReplication/MassReplication.Build.cs
--- a/Engine/Plugins/Runtime/MassGameplay/Source/MassReplication/MassReplication.Build.cs
+++ b/Engine/Plugins/Runtime/MassGameplay/Source/MassReplication/MassReplication.Build.cs
@@ -8,13 +8,6 @@ namespace UnrealBuildTool.Rules
 		{
 			PCHUsage = ModuleRules.PCHUsageMode.UseExplicitOrSharedPCHs;
 
-			PublicIncludePaths.AddRange(
-				new string[] {
-				"Runtime/AIModule/Public",
-				ModuleDirectory + "/Public",
-				}
-			);
-
 			PublicDependencyModuleNames.AddRange(
 				new string[] {
 					"Core",
diff --git a/Engine/Plugins/Runtime/MassGameplay/Source/MassRepresentation/MassRepresentation.Build.cs b/Engine/Plugins/Runtime/MassGameplay/Source/MassRepresentation/MassRepresentation.Build.cs
--- a/Engine/Plugins/Runtime/MassGameplay/Source/MassRepresentation/MassRepresentation.Build.cs
+++ b/Engine/Plugins/Runtime/MassGameplay/Source/MassRepresentation/MassRepresentation.Build.cs
@@ -8,13 +8,6 @@ namespace UnrealBuildTool.Rules
 		{
 			PCHUsage = ModuleRules.PCHUsageMode.UseExplicitOrSharedPCHs;
 
-			PublicIncludePaths.AddRange(
-				new string[] {
-					"Runtime/AIModule/Public",
-					ModuleDirectory + "/Public",
-				}
-			);
-
 			PublicDependencyModuleNames.AddRange(
 				new string[] {
 					"Core",
diff --git a/Engine/Plugins/Runtime/MassGameplay/Source/MassSignals/MassSignals.Build.cs b/Engine/Plugins/Runtime/MassGameplay/Source/MassSignals/MassSignals.Build.cs
--- a/Engine/Plugins/Runtime/MassGameplay/Source/MassSignals/MassSignals.Build.cs
+++ b/Engine/Plugins/Runtime/MassGameplay/Source/MassSignals/MassSignals.Build.cs
@@ -8,13 +8,6 @@ namespace UnrealBuildTool.Rules
 		{
 			PCHUsage = ModuleRules.PCHUsageMode.UseExplicitOrSharedPCHs;
 
-			PublicIncludePaths.AddRange(
-				new string[] {
-					"Runtime/AIModule/Public",
-					ModuleDirectory + "/Public",
-				}
-			);
-
 			PublicDependencyModuleNames.AddRange(
 				new string[] {
 					"Core",
diff --git a/Engine/Plugins/Runtime/MassGameplay/Source/MassSimulation/MassSimulation.Build.cs b/Engine/Plugins/Runtime/MassGameplay/Source/MassSimulation/MassSimulation.Build.cs
--- a/Engine/Plugins/Runtime/MassGameplay/Source/MassSimulation/MassSimulation.Build.cs
+++ b/Engine/Plugins/Runtime/MassGameplay/Source/MassSimulation/MassSimulation.Build.cs
@@ -8,13 +8,6 @@ namespace UnrealBuildTool.Rules
 		{
 			PCHUsage = ModuleRules.PCHUsageMode.UseExplicitOrSharedPCHs;
 
-			PublicIncludePaths.AddRange(
-				new string[] {
-				"Runtime/AIModule/Public",
-				ModuleDirectory + "/Public",
-				}
-			);
-
 			PublicDependencyModuleNames.AddRange(
 				new string[] {
 					"Core",
diff --git a/Engine/Plugins/Runtime/MassGameplay/Source/MassSmartObjects/MassSmartObjects.Build.cs b/Engine/Plugins/Runtime/MassGameplay/Source/MassSmartObjects/MassSmartObjects.Build.cs
--- a/Engine/Plugins/Runtime/MassGameplay/Source/MassSmartObjects/MassSmartObjects.Build.cs
+++ b/Engine/Plugins/Runtime/MassGameplay/Source/MassSmartObjects/MassSmartObjects.Build.cs
@@ -8,13 +8,6 @@ namespace UnrealBuildTool.Rules
 		{
 			PCHUsage = ModuleRules.PCHUsageMode.UseExplicitOrSharedPCHs;
 
-			PublicIncludePaths.AddRange(
-				new string[] {
-				"Runtime/AIModule/Public",
-				ModuleDirectory + "/Public",
-				}
-			);
-
 			PublicDependencyModuleNames.AddRange(
 				new string[] {
 					"MassEntity",
diff --git a/Engine/Plugins/Runtime/ZoneGraph/Source/ZoneGraph/ZoneGraph.Build.cs b/Engine/Plugins/Runtime/ZoneGraph/Source/ZoneGraph/ZoneGraph.Build.cs
--- a/Engine/Plugins/Runtime/ZoneGraph/Source/ZoneGraph/ZoneGraph.Build.cs
+++ b/Engine/Plugins/Runtime/ZoneGraph/Source/ZoneGraph/ZoneGraph.Build.cs
@@ -8,13 +8,6 @@ namespace UnrealBuildTool.Rules
 		{
 			PCHUsage = ModuleRules.PCHUsageMode.UseExplicitOrSharedPCHs;
 
-			PublicIncludePaths.AddRange(
-				new string[] {
-					"Runtime/AIModule/Public",
-					ModuleDirectory + "/Public",
-				}
-			);
-
 			PublicDependencyModuleNames.AddRange(
 				new string[] {
 					"Core",
diff --git a/Engine/Plugins/Runtime/ZoneGraph/Source/ZoneGraphDebug/ZoneGraphDebug.Build.cs b/Engine/Plugins/Runtime/ZoneGraph/Source/ZoneGraphDebug/ZoneGraphDebug.Build.cs
--- a/Engine/Plugins/Runtime/ZoneGraph/Source/ZoneGraphDebug/ZoneGraphDebug.Build.cs
+++ b/Engine/Plugins/Runtime/ZoneGraph/Source/ZoneGraphDebug/ZoneGraphDebug.Build.cs
@@ -8,13 +8,6 @@ namespace UnrealBuildTool.Rules
 		{
 			PCHUsage = ModuleRules.PCHUsageMode.UseExplicitOrSharedPCHs;
 
-			PublicIncludePaths.AddRange(
-			new string[] {
-				"Runtime/AIModule/Public",
-				ModuleDirectory + "/Public",
-			}
-			);
-
 			PublicDependencyModuleNames.AddRange(
 			new string[] {
 				"Core",
diff --git a/Engine/Plugins/Runtime/ZoneGraphAnnotations/Source/ZoneGraphAnnotations/ZoneGraphAnnotations.Build.cs b/Engine/Plugins/Runtime/ZoneGraphAnnotations/Source/ZoneGraphAnnotations/ZoneGraphAnnotations.Build.cs
--- a/Engine/Plugins/Runtime/ZoneGraphAnnotations/Source/ZoneGraphAnnotations/ZoneGraphAnnotations.Build.cs
+++ b/Engine/Plugins/Runtime/ZoneGraphAnnotations/Source/ZoneGraphAnnotations/ZoneGraphAnnotations.Build.cs
@@ -8,13 +8,6 @@ namespace UnrealBuildTool.Rules
 		{
 			PCHUsage = ModuleRules.PCHUsageMode.UseExplicitOrSharedPCHs;
 
-			PublicIncludePaths.AddRange(
-				new string[] {
-					"Runtime/AIModule/Public",
-					ModuleDirectory + "/Public",
-				}
-			);
-
 			PublicDependencyModuleNames.AddRange(
 				new string[] {
 					"Core",
