Fix spelling error
#rnx

diff --git a/Engine/Plugins/AI/MLAdapter/Source/MLAdapter/Private/MLAdapterLibrarian.cpp b/Engine/Plugins/AI/MLAdapter/Source/MLAdapter/Private/MLAdapterLibrarian.cpp
--- a/Engine/Plugins/AI/MLAdapter/Source/MLAdapter/Private/MLAdapterLibrarian.cpp
+++ b/Engine/Plugins/AI/MLAdapter/Source/MLAdapter/Private/MLAdapterLibrarian.cpp
@@ -24,7 +24,7 @@ namespace FLibrarianHelper
 				, *Class->GetName()
 				, Class->HasAnyClassFlags(CLASS_Abstract) ? TEXT("Abstract, ") : TEXT("")
 				, Class->HasAnyClassFlags(CLASS_Deprecated) ? TEXT("Deprecated, ") : TEXT("")
-				, Class->HasAnyClassFlags(CLASS_NewerVersionExists) ? TEXT("NewerVesionExists, ") : TEXT("")
+				, Class->HasAnyClassFlags(CLASS_NewerVersionExists) ? TEXT("NewerVersionExists, ") : TEXT("")
 			);
 			return false;
 		}
