Fixed a compile error caused by change to the return value of AController::GetPawn().

diff --git a/Engine/Plugins/AI/MLAdapter/Source/MLAdapter/Private/Agents/MLAdapterAgentElement.cpp b/Engine/Plugins/AI/MLAdapter/Source/MLAdapter/Private/Agents/MLAdapterAgentElement.cpp
--- a/Engine/Plugins/AI/MLAdapter/Source/MLAdapter/Private/Agents/MLAdapterAgentElement.cpp
+++ b/Engine/Plugins/AI/MLAdapter/Source/MLAdapter/Private/Agents/MLAdapterAgentElement.cpp
@@ -47,7 +47,7 @@ APawn* UMLAdapterAgentElement::GetPawnAvatar() const
 	AActor* Avatar = GetAvatar();
 	APawn* Pawn = Cast<APawn>(Avatar);
 	AController* Controller = Cast<AController>(Avatar);
-	return Pawn ? Pawn : (Controller ? Controller->GetPawn() : nullptr);
+	return Pawn ? Pawn : (Controller ? Controller->GetPawn().Get() : nullptr);
 }
 
 AController* UMLAdapterAgentElement::GetControllerAvatar() const
