fix msgpack-c not detecting byte order on Windows Arm64

#rnx

diff --git a/Engine/Plugins/AI/MLAdapter/Source/MLAdapter/MLAdapter.Build.cs b/Engine/Plugins/AI/MLAdapter/Source/MLAdapter/MLAdapter.Build.cs
--- a/Engine/Plugins/AI/MLAdapter/Source/MLAdapter/MLAdapter.Build.cs
+++ b/Engine/Plugins/AI/MLAdapter/Source/MLAdapter/MLAdapter.Build.cs
@@ -39,6 +39,11 @@ namespace UnrealBuildTool.Rules
             {
                 PublicDefinitions.Add("WITH_RPCLIB=1");
 				PrivateDependencyModuleNames.Add("RPCLib");
+
+				if (Target.Platform == UnrealTargetPlatform.Win64 && Target.Architecture == UnrealArch.Arm64)
+				{
+					PrivateDefinitions.Add("__arm64"); // expected by rpclib headers
+				}
 			}
             else
             {
diff --git a/Engine/Plugins/AI/MLAdapter/Source/MLAdapterTestSuite/MLAdapterTestSuite.Build.cs b/Engine/Plugins/AI/MLAdapter/Source/MLAdapterTestSuite/MLAdapterTestSuite.Build.cs
--- a/Engine/Plugins/AI/MLAdapter/Source/MLAdapterTestSuite/MLAdapterTestSuite.Build.cs
+++ b/Engine/Plugins/AI/MLAdapter/Source/MLAdapterTestSuite/MLAdapterTestSuite.Build.cs
@@ -40,6 +40,11 @@ namespace UnrealBuildTool.Rules
 			{
 				PublicDefinitions.Add("WITH_RPCLIB=1");
 				PrivateDependencyModuleNames.Add("RPCLib");
+
+				if (Target.Platform == UnrealTargetPlatform.Win64 && Target.Architecture == UnrealArch.Arm64)
+				{
+					PrivateDefinitions.Add("__arm64"); // expected by rpclib headers
+				}
 			}
 			else
             {
