MLAdapter: Skip unix.Malloc check

#rnx

diff --git a/Engine/Plugins/AI/MLAdapter/Source/MLAdapter/MLAdapter.Build.cs b/Engine/Plugins/AI/MLAdapter/Source/MLAdapter/MLAdapter.Build.cs
--- a/Engine/Plugins/AI/MLAdapter/Source/MLAdapter/MLAdapter.Build.cs
+++ b/Engine/Plugins/AI/MLAdapter/Source/MLAdapter/MLAdapter.Build.cs
@@ -6,8 +6,11 @@ namespace UnrealBuildTool.Rules
     {
         public MLAdapter(ReadOnlyTargetRules Target) : base(Target)
         {
-            // rcplib is using exceptions so we have to enable that
-            bEnableExceptions = true;
+			// recheck this if rpclib is updated
+			StaticAnalyzerDisabledCheckers.Add("unix.Malloc");
+
+			// rcplib is using exceptions so we have to enable that
+			bEnableExceptions = true;
             PCHUsage = ModuleRules.PCHUsageMode.UseExplicitOrSharedPCHs;
 
             PublicDependencyModuleNames.AddRange(
